Fix for stat metadata mismatch warning during PIE.

diff --git a/Engine/Plugins/Online/OnlineSubsystemNull/Source/Private/OnlineSubsystemNull.cpp b/Engine/Plugins/Online/OnlineSubsystemNull/Source/Private/OnlineSubsystemNull.cpp
--- a/Engine/Plugins/Online/OnlineSubsystemNull/Source/Private/OnlineSubsystemNull.cpp
+++ b/Engine/Plugins/Online/OnlineSubsystemNull/Source/Private/OnlineSubsystemNull.cpp
@@ -10,6 +10,8 @@
 #include "VoiceInterfaceImpl.h"
 #include "OnlineAchievementsInterfaceNull.h"
 
+FThreadSafeCounter FOnlineSubsystemNull::TaskCounter;
+
 IOnlineSessionPtr FOnlineSubsystemNull::GetSessionInterface() const
 {
 	return SessionInterface;
@@ -164,7 +166,7 @@ bool FOnlineSubsystemNull::Init()
 		// Create the online async task thread
 		OnlineAsyncTaskThreadRunnable = new FOnlineAsyncTaskManagerNull(this);
 		check(OnlineAsyncTaskThreadRunnable);
-		OnlineAsyncTaskThread = FRunnableThread::Create(OnlineAsyncTaskThreadRunnable, *FString::Printf(TEXT("OnlineAsyncTaskThreadNull %s"), *InstanceName.ToString()), 128 * 1024, TPri_Normal);
+		OnlineAsyncTaskThread = FRunnableThread::Create(OnlineAsyncTaskThreadRunnable, *FString::Printf(TEXT("OnlineAsyncTaskThreadNull %s(%d)"), *InstanceName.ToString(), TaskCounter.Increment()), 128 * 1024, TPri_Normal);
 		check(OnlineAsyncTaskThread);
 		UE_LOG_ONLINE(Verbose, TEXT("Created thread (ID:%d)."), OnlineAsyncTaskThread->GetThreadID());
 
diff --git a/Engine/Plugins/Online/OnlineSubsystemNull/Source/Public/OnlineSubsystemNull.h b/Engine/Plugins/Online/OnlineSubsystemNull/Source/Public/OnlineSubsystemNull.h
--- a/Engine/Plugins/Online/OnlineSubsystemNull/Source/Public/OnlineSubsystemNull.h
+++ b/Engine/Plugins/Online/OnlineSubsystemNull/Source/Public/OnlineSubsystemNull.h
@@ -5,6 +5,7 @@
 #include "CoreMinimal.h"
 #include "OnlineSubsystemImpl.h"
 #include "OnlineSubsystemNullPackage.h"
+#include "HAL/ThreadSafeCounter.h"
 
 class FOnlineAchievementsNull;
 class FOnlineIdentityNull;
@@ -131,6 +132,9 @@ private:
 
 	/** Online async task thread */
 	class FRunnableThread* OnlineAsyncTaskThread;
+
+	// task counter, used to generate unique thread names for each task
+	static FThreadSafeCounter TaskCounter;
 };
 
 typedef TSharedPtr<FOnlineSubsystemNull, ESPMode::ThreadSafe> FOnlineSubsystemNullPtr;
