MCP Test Sanitizer

When running with -testsanitizer on the command line (and without a more specific platform message sanitizer available), the game will now automatically reject the word "dirty", ignoring case when performing OSS Message Sanitization.

This will help developers and testers verify on development platforms that OSS-specific message sanitizers on more esoteric platforms will be working.  This test sanitizer also uses a built-in multi-second delay help emulate situations where platform sanitizers can be of variable speed.

diff --git a/Engine/Plugins/Online/OnlineSubsystem/Source/Private/OnlineSubsystemImpl.cpp b/Engine/Plugins/Online/OnlineSubsystem/Source/Private/OnlineSubsystemImpl.cpp
--- a/Engine/Plugins/Online/OnlineSubsystem/Source/Private/OnlineSubsystemImpl.cpp
+++ b/Engine/Plugins/Online/OnlineSubsystem/Source/Private/OnlineSubsystemImpl.cpp
@@ -223,6 +223,16 @@ IMessageSanitizerPtr FOnlineSubsystemImpl::GetMessageSanitizer(int32 LocalUserNu
 	{
 		MessageSanitizer = SanitizerSubsystem->GetMessageSanitizer(LocalUserNum, OutAuthTypeToExclude);
 	}
+
+	if (!MessageSanitizer && FParse::Param(FCommandLine::Get(), TEXT("testsanitizer")))
+	{
+		SanitizerSubsystem = IOnlineSubsystem::Get(NULL_SUBSYSTEM);
+		if (SanitizerSubsystem)
+		{
+			return SanitizerSubsystem->GetMessageSanitizer(LocalUserNum, OutAuthTypeToExclude);
+		}
+	}
+
 	return MessageSanitizer;
 }
 
diff --git a/Engine/Plugins/Online/OnlineSubsystemNull/Source/Private/OnlineMessageSanitizerNull.cpp b/Engine/Plugins/Online/OnlineSubsystemNull/Source/Private/OnlineMessageSanitizerNull.cpp
new file mode 100644
--- /dev/null
+++ b/Engine/Plugins/Online/OnlineSubsystemNull/Source/Private/OnlineMessageSanitizerNull.cpp
@@ -0,0 +1,73 @@
+#include "OnlineMessageSanitizerNull.h"
+#include "Containers/Ticker.h"
+
+FMessageSanitizerNull::FMessageSanitizerNull()
+	: RequestId(0)
+{
+}
+
+FMessageSanitizerNull::~FMessageSanitizerNull()
+{
+	for (TPair<int32, FDelegateHandle>& Pair : RequestIdHandleMap)
+	{
+		FTicker::GetCoreTicker().RemoveTicker(Pair.Value);
+	}
+}
+
+void FMessageSanitizerNull::SanitizeDisplayName(const FString& DisplayName, const FOnMessageProcessed& CompletionDelegate)
+{
+	float Delay = 3.0f;
+	int32 ThisRequest = RequestId++;
+
+	FDelegateHandle TickerHandle = FTicker::GetCoreTicker().AddTicker(FTickerDelegate::CreateLambda([this, DisplayName, CompletionDelegate, ThisRequest](float) -> bool
+	{
+		FString CleanString;
+		PerformSanitize(DisplayName, CleanString);
+		CompletionDelegate.ExecuteIfBound(true, CleanString);
+		UE_LOG_ONLINE(Verbose, TEXT("Completed MCP Test sanitize of string %s (result: %s)"), *DisplayName, *CleanString);
+		RequestIdHandleMap.Remove(ThisRequest);
+		return false;
+	}), Delay);
+
+	RequestIdHandleMap.Add(ThisRequest, TickerHandle);
+}
+
+void FMessageSanitizerNull::SanitizeDisplayNames(const TArray<FString>& DisplayNames, const FOnMessageArrayProcessed& CompletionDelegate)
+{
+	float Delay = 3.0f;
+	int32 ThisRequest = RequestId++;
+
+	FDelegateHandle TickerHandle = FTicker::GetCoreTicker().AddTicker(FTickerDelegate::CreateLambda([this, DisplayNames, CompletionDelegate, ThisRequest](float) -> bool
+	{
+		TArray<FString> OutStrings;
+		for (const FString& String : DisplayNames)
+		{
+			FString CleanString;
+			PerformSanitize(String, CleanString);
+			OutStrings.Add(MoveTemp(CleanString));
+		}
+		CompletionDelegate.ExecuteIfBound(true, OutStrings);
+		UE_LOG_ONLINE(Verbose, TEXT("Completed MCP Test sanitize of %d strings."), DisplayNames.Num());
+		RequestIdHandleMap.Remove(ThisRequest);
+		return false;
+	}), Delay);
+
+	RequestIdHandleMap.Add(ThisRequest, TickerHandle);
+}
+
+void FMessageSanitizerNull::PerformSanitize(const FString& RawString, FString& OutCleanString) const
+{
+	OutCleanString = RawString.Replace(TEXT("dirty"), TEXT("d***y"));
+}
+
+void FMessageSanitizerNull::QueryBlockedUser(int32 LocalUserNum, const FString& FromUserIdStr, const FString& FromPlatform, const FOnQueryUserBlockedResponse& CompletionDelegate)
+{
+	FBlockedQueryResult Result;
+	Result.UserId = FromUserIdStr;
+	Result.bIsBlocked = false;
+	CompletionDelegate.ExecuteIfBound(Result);
+}
+
+void FMessageSanitizerNull::ResetBlockedUserCache()
+{
+}
\ No newline at end of file
diff --git a/Engine/Plugins/Online/OnlineSubsystemNull/Source/Private/OnlineMessageSanitizerNull.h b/Engine/Plugins/Online/OnlineSubsystemNull/Source/Private/OnlineMessageSanitizerNull.h
new file mode 100644
--- /dev/null
+++ b/Engine/Plugins/Online/OnlineSubsystemNull/Source/Private/OnlineMessageSanitizerNull.h
@@ -0,0 +1,31 @@
+// Copyright Epic Games, Inc. All Rights Reserved.
+
+#pragma once
+
+#include "OnlineSubsystem.h"
+#include "Interfaces/IMessageSanitizerInterface.h"
+
+/**
+ * Implements a degenerate/dummy interface chat message sanitization system for testing
+ */
+class FMessageSanitizerNull
+	: public IMessageSanitizer
+{
+public:
+
+	// IMessageSanitizer
+	virtual void SanitizeDisplayName(const FString& DisplayName, const FOnMessageProcessed& CompletionDelegate) override;
+	virtual void SanitizeDisplayNames(const TArray<FString>& DisplayNames, const FOnMessageArrayProcessed& CompletionDelegate) override;
+	virtual void QueryBlockedUser(int32 LocalUserNum, const FString& FromUserIdStr, const FString& FromPlatform, const FOnQueryUserBlockedResponse& InCompletionDelegate) override;
+	virtual void ResetBlockedUserCache() override;
+	// FMessageSanitizerNull
+
+	FMessageSanitizerNull();
+	virtual ~FMessageSanitizerNull();
+
+private:
+	int32 RequestId;
+
+	void PerformSanitize(const FString& RawString, FString& OutCleanString) const;
+	TMap<int32, FDelegateHandle> RequestIdHandleMap;
+};
\ No newline at end of file
diff --git a/Engine/Plugins/Online/OnlineSubsystemNull/Source/Private/OnlineSubsystemNull.cpp b/Engine/Plugins/Online/OnlineSubsystemNull/Source/Private/OnlineSubsystemNull.cpp
--- a/Engine/Plugins/Online/OnlineSubsystemNull/Source/Private/OnlineSubsystemNull.cpp
+++ b/Engine/Plugins/Online/OnlineSubsystemNull/Source/Private/OnlineSubsystemNull.cpp
@@ -11,6 +11,7 @@
 #include "OnlineAchievementsInterfaceNull.h"
 #include "OnlineStoreV2InterfaceNull.h"
 #include "OnlinePurchaseInterfaceNull.h"
+#include "OnlineMessageSanitizerNull.h"
 
 FThreadSafeCounter FOnlineSubsystemNull::TaskCounter;
 
@@ -149,6 +150,11 @@ IOnlineTournamentPtr FOnlineSubsystemNull::GetTournamentInterface() const
 	return nullptr;
 }
 
+IMessageSanitizerPtr FOnlineSubsystemNull::GetMessageSanitizer(int32 LocalUserNum, FString& OutAuthTypeToExclude) const
+{
+	return MessageSanitizerInterface;
+}
+
 bool FOnlineSubsystemNull::Tick(float DeltaTime)
 {
 	if (!FOnlineSubsystemImpl::Tick(DeltaTime))
@@ -194,6 +200,7 @@ bool FOnlineSubsystemNull::Init()
 		VoiceInterface = MakeShareable(new FOnlineVoiceImpl(this));
 		StoreV2Interface = MakeShareable(new FOnlineStoreV2Null(*this));
 		PurchaseInterface = MakeShareable(new FOnlinePurchaseNull(*this));
+		MessageSanitizerInterface = MakeShareable(new FMessageSanitizerNull());
 	}
 	else
 	{
@@ -242,7 +249,8 @@ bool FOnlineSubsystemNull::Shutdown()
 	DESTRUCT_INTERFACE(IdentityInterface);
 	DESTRUCT_INTERFACE(LeaderboardsInterface);
 	DESTRUCT_INTERFACE(SessionInterface);
-	
+	DESTRUCT_INTERFACE(MessageSanitizerInterface);
+
 #undef DESTRUCT_INTERFACE
 	
 	return true;
diff --git a/Engine/Plugins/Online/OnlineSubsystemNull/Source/Public/OnlineSubsystemNull.h b/Engine/Plugins/Online/OnlineSubsystemNull/Source/Public/OnlineSubsystemNull.h
--- a/Engine/Plugins/Online/OnlineSubsystemNull/Source/Public/OnlineSubsystemNull.h
+++ b/Engine/Plugins/Online/OnlineSubsystemNull/Source/Public/OnlineSubsystemNull.h
@@ -25,6 +25,7 @@ typedef TSharedPtr<class FOnlineIdentityNull, ESPMode::ThreadSafe> FOnlineIdenti
 typedef TSharedPtr<class FOnlineAchievementsNull, ESPMode::ThreadSafe> FOnlineAchievementsNullPtr;
 typedef TSharedPtr<class FOnlineStoreV2Null, ESPMode::ThreadSafe> FOnlineStoreV2NullPtr;
 typedef TSharedPtr<class FOnlinePurchaseNull, ESPMode::ThreadSafe> FOnlinePurchaseNullPtr;
+typedef TSharedPtr<class FMessageSanitizerNull, ESPMode::ThreadSafe> FMessageSanitizerNullPtr;
 
 /**
  *	OnlineSubsystemNull - Implementation of the online subsystem for Null services
@@ -64,6 +65,7 @@ public:
 	virtual IOnlineStatsPtr GetStatsInterface() const override;
 	virtual IOnlineTurnBasedPtr GetTurnBasedInterface() const override;
 	virtual IOnlineTournamentPtr GetTournamentInterface() const override;
+	virtual IMessageSanitizerPtr GetMessageSanitizer(int32 LocalUserNum, FString& OutAuthTypeToExclude) const override;
 
 	virtual bool Init() override;
 	virtual bool Shutdown() override;
@@ -90,6 +92,7 @@ PACKAGE_SCOPE:
 		IdentityInterface(nullptr),
 		AchievementsInterface(nullptr),
 		StoreV2Interface(nullptr),
+		MessageSanitizerInterface(nullptr),
 		OnlineAsyncTaskThreadRunnable(nullptr),
 		OnlineAsyncTaskThread(nullptr)
 	{}
@@ -120,6 +123,9 @@ private:
 	/** Interface for purchases */
 	FOnlinePurchaseNullPtr PurchaseInterface;
 
+	/** Interface for message sanitizing */
+	FMessageSanitizerNullPtr MessageSanitizerInterface;
+
 	/** Online async task runnable */
 	class FOnlineAsyncTaskManagerNull* OnlineAsyncTaskThreadRunnable;
 
