Added missing copyright text.

diff --git a/Engine/Plugins/Online/OnlineSubsystemNull/Source/Private/OnlineMessageSanitizerNull.cpp b/Engine/Plugins/Online/OnlineSubsystemNull/Source/Private/OnlineMessageSanitizerNull.cpp
--- a/Engine/Plugins/Online/OnlineSubsystemNull/Source/Private/OnlineMessageSanitizerNull.cpp
+++ b/Engine/Plugins/Online/OnlineSubsystemNull/Source/Private/OnlineMessageSanitizerNull.cpp
@@ -1,3 +1,5 @@
+// Copyright Epic Games, Inc. All Rights Reserved.
+
 #include "OnlineMessageSanitizerNull.h"
 #include "Containers/Ticker.h"
 
