Getting OssNull to compile w/o engine. This mainly just involves removing voice chat when engine is not included (since it uses VoiceEngineImpl- if this is a requirement for OssUnitTests in the future, we may just have to move this to its own module)

diff --git a/Engine/Plugins/Online/OnlineSubsystemNull/Source/OnlineSubsystemNull.Build.cs b/Engine/Plugins/Online/OnlineSubsystemNull/Source/OnlineSubsystemNull.Build.cs
--- a/Engine/Plugins/Online/OnlineSubsystemNull/Source/OnlineSubsystemNull.Build.cs
+++ b/Engine/Plugins/Online/OnlineSubsystemNull/Source/OnlineSubsystemNull.Build.cs
@@ -10,21 +10,30 @@ public class OnlineSubsystemNull : ModuleRules
 		PrivateDefinitions.Add("ONLINESUBSYSTEMNULL_PACKAGE=1");
 		PCHUsage = ModuleRules.PCHUsageMode.UseExplicitOrSharedPCHs;
 
-        PublicDependencyModuleNames.AddRange(
-            new string[] {
-                "OnlineSubsystemUtils"
-            }
+
+        if (Target.bCompileAgainstEngine)
+        {
+            PrivateDependencyModuleNames.AddRange(
+                new string[] {
+                    "Engine"
+                }
             );
 
+            PublicDependencyModuleNames.AddRange(
+			   new string[] {
+					"OnlineSubsystemUtils"
+			   }
+		   );
+		}
+
         PrivateDependencyModuleNames.AddRange(
 			new string[] {
-				"Core", 
-				"CoreUObject", 
-				"Engine", 
-				"Sockets", 
+				"Core",
+				"Sockets",
+				"CoreUObject",
 				"OnlineSubsystem", 
 				"Json"
 			}
-			);
+		);
 	}
 }
diff --git a/Engine/Plugins/Online/OnlineSubsystemNull/Source/Private/OnlineSessionInterfaceNull.cpp b/Engine/Plugins/Online/OnlineSubsystemNull/Source/Private/OnlineSessionInterfaceNull.cpp
--- a/Engine/Plugins/Online/OnlineSubsystemNull/Source/Private/OnlineSessionInterfaceNull.cpp
+++ b/Engine/Plugins/Online/OnlineSubsystemNull/Source/Private/OnlineSessionInterfaceNull.cpp
@@ -5,10 +5,14 @@
 #include "OnlineSubsystem.h"
 #include "OnlineSubsystemNull.h"
 #include "OnlineSubsystemNullTypes.h"
+#if WITH_ENGINE
 #include "OnlineSubsystemUtils.h"
+#endif //WITH_ENGINE
 #include "OnlineAsyncTaskManager.h"
 #include "SocketSubsystem.h"
 #include "NboSerializerNull.h"
+#include "Interfaces/OnlineIdentityInterface.h"
+#include "UObject/CoreNet.h"
 
 
 
@@ -37,7 +41,10 @@ void FOnlineSessionInfoNull::Init(const FOnlineSubsystemNull& Subsystem)
 	}
 
 	// Now set the port that was configured
+#if WITH_ENGINE
+	// jntodo: can we put anything here if we don't have engine code?
 	HostAddr->SetPort(GetPortFromNetDriver(Subsystem.GetInstanceName()));
+#endif //WITH_ENGINE
 
 	FGuid OwnerGuid;
 	FPlatformMisc::CreateGuid(OwnerGuid);
@@ -807,6 +814,7 @@ FOnlineSessionSettings* FOnlineSessionNull::GetSessionSettings(FName SessionName
 
 void FOnlineSessionNull::RegisterLocalPlayers(FNamedOnlineSession* Session)
 {
+#if WITH_ENGINE
 	if (!NullSubsystem->IsDedicated())
 	{
 		IOnlineVoicePtr VoiceInt = NullSubsystem->GetVoiceInterface();
@@ -819,10 +827,12 @@ void FOnlineSessionNull::RegisterLocalPlayers(FNamedOnlineSession* Session)
 			}
 		}
 	}
+#endif //WITH_ENGINE
 }
 
 void FOnlineSessionNull::RegisterVoice(const FUniqueNetId& PlayerId)
 {
+#if WITH_ENGINE
 	IOnlineVoicePtr VoiceInt = NullSubsystem->GetVoiceInterface();
 	if (VoiceInt.IsValid())
 	{
@@ -836,10 +846,12 @@ void FOnlineSessionNull::RegisterVoice(const FUniqueNetId& PlayerId)
 			VoiceInt->ProcessMuteChangeNotification();
 		}
 	}
+#endif //WITH_ENGINE
 }
 
 void FOnlineSessionNull::UnregisterVoice(const FUniqueNetId& PlayerId)
 {
+#if WITH_ENGINE
 	IOnlineVoicePtr VoiceInt = NullSubsystem->GetVoiceInterface();
 	if (VoiceInt.IsValid())
 	{
@@ -851,6 +863,7 @@ void FOnlineSessionNull::UnregisterVoice(const FUniqueNetId& PlayerId)
 			}
 		}
 	}
+#endif //WITH_ENGINE
 }
 
 bool FOnlineSessionNull::RegisterPlayer(FName SessionName, const FUniqueNetId& PlayerId, bool bWasInvited)
@@ -1241,12 +1254,14 @@ void FOnlineSessionNull::UnregisterLocalPlayer(const FUniqueNetId& PlayerId, FNa
 
 void FOnlineSessionNull::SetPortFromNetDriver(const FOnlineSubsystemNull& Subsystem, const TSharedPtr<FOnlineSessionInfo>& SessionInfo)
 {
+#if WITH_ENGINE
 	auto NetDriverPort = GetPortFromNetDriver(Subsystem.GetInstanceName());
 	auto SessionInfoNull = StaticCastSharedPtr<FOnlineSessionInfoNull>(SessionInfo);
 	if (SessionInfoNull.IsValid() && SessionInfoNull->HostAddr.IsValid())
 	{
 		SessionInfoNull->HostAddr->SetPort(NetDriverPort);
 	}
+#endif //WITH_ENGINE
 }
 
 bool FOnlineSessionNull::IsHost(const FNamedOnlineSession& Session) const
diff --git a/Engine/Plugins/Online/OnlineSubsystemNull/Source/Private/OnlineSubsystemNull.cpp b/Engine/Plugins/Online/OnlineSubsystemNull/Source/Private/OnlineSubsystemNull.cpp
--- a/Engine/Plugins/Online/OnlineSubsystemNull/Source/Private/OnlineSubsystemNull.cpp
+++ b/Engine/Plugins/Online/OnlineSubsystemNull/Source/Private/OnlineSubsystemNull.cpp
@@ -7,13 +7,16 @@
 #include "OnlineSessionInterfaceNull.h"
 #include "OnlineLeaderboardInterfaceNull.h"
 #include "OnlineIdentityNull.h"
-#include "VoiceInterfaceNull.h"
 #include "OnlineAchievementsInterfaceNull.h"
 #include "OnlineStoreV2InterfaceNull.h"
 #include "OnlinePurchaseInterfaceNull.h"
 #include "OnlineMessageSanitizerNull.h"
 #include "Stats/Stats.h"
 
+#if WITH_ENGINE
+#include "VoiceInterfaceNull.h"
+#endif //WITH_ENGINE
+
 FThreadSafeCounter FOnlineSubsystemNull::TaskCounter;
 
 IOnlineSessionPtr FOnlineSubsystemNull::GetSessionInterface() const
@@ -58,6 +61,7 @@ IOnlineLeaderboardsPtr FOnlineSubsystemNull::GetLeaderboardsInterface() const
 
 IOnlineVoicePtr FOnlineSubsystemNull::GetVoiceInterface() const
 {
+#if WITH_ENGINE
 	if (VoiceInterface.IsValid() && !bVoiceInterfaceInitialized)
 	{	
 		if (!VoiceInterface->Init())
@@ -69,6 +73,9 @@ IOnlineVoicePtr FOnlineSubsystemNull::GetVoiceInterface() const
 	}
 
 	return VoiceInterface;
+#else //WITH_ENGINE
+	return nullptr;
+#endif //WITH_ENGINE
 }
 
 IOnlineExternalUIPtr FOnlineSubsystemNull::GetExternalUIInterface() const
@@ -175,10 +182,12 @@ bool FOnlineSubsystemNull::Tick(float DeltaTime)
 		SessionInterface->Tick(DeltaTime);
 	}
 
+#if WITH_ENGINE
 	if (VoiceInterface.IsValid() && bVoiceInterfaceInitialized)
 	{
 		VoiceInterface->Tick(DeltaTime);
 	}
+#endif //WITH_ENGINE
 
 	return true;
 }
@@ -200,7 +209,9 @@ bool FOnlineSubsystemNull::Init()
 		LeaderboardsInterface = MakeShareable(new FOnlineLeaderboardsNull(this));
 		IdentityInterface = MakeShareable(new FOnlineIdentityNull(this));
 		AchievementsInterface = MakeShareable(new FOnlineAchievementsNull(this));
+#if WITH_ENGINE
 		VoiceInterface = MakeShareable(new FOnlineVoiceImpl(this));
+#endif //WITH_ENGINE
 		StoreV2Interface = MakeShareable(new FOnlineStoreV2Null(*this));
 		PurchaseInterface = MakeShareable(new FOnlinePurchaseNull(*this));
 		MessageSanitizerInterface = MakeShareable(new FMessageSanitizerNull());
@@ -232,10 +243,12 @@ bool FOnlineSubsystemNull::Shutdown()
 		OnlineAsyncTaskThreadRunnable = nullptr;
 	}
 
+#if WITH_ENGINE
 	if (VoiceInterface.IsValid() && bVoiceInterfaceInitialized)
 	{
 		VoiceInterface->Shutdown();
 	}
+#endif //WITH_ENGINE
 	
 #define DESTRUCT_INTERFACE(Interface) \
 	if (Interface.IsValid()) \
diff --git a/Engine/Plugins/Online/OnlineSubsystemNull/Source/Public/OnlineSubsystemNull.h b/Engine/Plugins/Online/OnlineSubsystemNull/Source/Public/OnlineSubsystemNull.h
--- a/Engine/Plugins/Online/OnlineSubsystemNull/Source/Public/OnlineSubsystemNull.h
+++ b/Engine/Plugins/Online/OnlineSubsystemNull/Source/Public/OnlineSubsystemNull.h
@@ -19,13 +19,15 @@ typedef TSharedPtr<class FOnlineProfileNull, ESPMode::ThreadSafe> FOnlineProfile
 typedef TSharedPtr<class FOnlineFriendsNull, ESPMode::ThreadSafe> FOnlineFriendsNullPtr;
 typedef TSharedPtr<class FOnlineUserCloudNull, ESPMode::ThreadSafe> FOnlineUserCloudNullPtr;
 typedef TSharedPtr<class FOnlineLeaderboardsNull, ESPMode::ThreadSafe> FOnlineLeaderboardsNullPtr;
-typedef TSharedPtr<class FOnlineVoiceImpl, ESPMode::ThreadSafe> FOnlineVoiceImplPtr;
 typedef TSharedPtr<class FOnlineExternalUINull, ESPMode::ThreadSafe> FOnlineExternalUINullPtr;
 typedef TSharedPtr<class FOnlineIdentityNull, ESPMode::ThreadSafe> FOnlineIdentityNullPtr;
 typedef TSharedPtr<class FOnlineAchievementsNull, ESPMode::ThreadSafe> FOnlineAchievementsNullPtr;
 typedef TSharedPtr<class FOnlineStoreV2Null, ESPMode::ThreadSafe> FOnlineStoreV2NullPtr;
 typedef TSharedPtr<class FOnlinePurchaseNull, ESPMode::ThreadSafe> FOnlinePurchaseNullPtr;
 typedef TSharedPtr<class FMessageSanitizerNull, ESPMode::ThreadSafe> FMessageSanitizerNullPtr;
+#if WITH_ENGINE
+typedef TSharedPtr<class FOnlineVoiceImpl, ESPMode::ThreadSafe> FOnlineVoiceImplPtr;
+#endif //WITH_ENGINE
 
 /**
  *	OnlineSubsystemNull - Implementation of the online subsystem for Null services
