Add config options to OnlineIdentityNull that let it emulate different types of login setups
Change default behavior to allow a null login with empty credentials, as this matches 99% of the platforms

diff --git a/Engine/Config/BaseEngine.ini b/Engine/Config/BaseEngine.ini
--- a/Engine/Config/BaseEngine.ini
+++ b/Engine/Config/BaseEngine.ini
@@ -1853,6 +1853,12 @@ bEnabled=false
 
 [OnlineSubsystemNull]
 bEnabled=true
+bAutoLoginAtStartup=true
+bRequireLoginCredentials=false
+bAddUserNumToNullId=false
+bForceStableNullId=false
+bForceOfflineMode=false
+
 Achievement_0_Id=null-ach-0
 Achievement_0_bIsHidden=false
 Achievement_0_Title="Achievement 0"
diff --git a/Engine/Plugins/Online/OnlineSubsystemNull/Source/Private/OnlineIdentityNull.cpp b/Engine/Plugins/Online/OnlineSubsystemNull/Source/Private/OnlineIdentityNull.cpp
--- a/Engine/Plugins/Online/OnlineSubsystemNull/Source/Private/OnlineIdentityNull.cpp
+++ b/Engine/Plugins/Online/OnlineSubsystemNull/Source/Private/OnlineIdentityNull.cpp
@@ -42,7 +42,7 @@ bool FUserOnlineAccountNull::SetUserAttribute(const FString& AttrName, const FSt
 	return false;
 }
 
-inline FString GenerateRandomUserId(int32 LocalUserNum)
+FString FOnlineIdentityNull::GenerateRandomUserId(int32 LocalUserNum)
 {
 	FString HostName;
 	if (!ISocketSubsystem::Get(PLATFORM_SOCKETSUBSYSTEM)->GetHostName(HostName))
@@ -53,18 +53,30 @@ inline FString GenerateRandomUserId(int32 LocalUserNum)
 		HostName = Addr->ToString(false);
 	}
 
-	const bool bForceUniqueId = FParse::Param( FCommandLine::Get(), TEXT( "StableNullID" ) );
-	
+	bool bUseStableNullId = bForceStableNullId;
+	FString UserSuffix;
+
+	if (bAddUserNumToNullId)
+	{
+		UserSuffix = FString::Printf(TEXT("-%d"), LocalUserNum);
+	}
+
 #if !(UE_BUILD_SHIPPING && WITH_EDITOR)
-	if ( ( GIsFirstInstance || bForceUniqueId ) && !GIsEditor )
+	if (GIsFirstInstance && !GIsEditor)
 	{
-		// When possible, return a stable user id
-		return FString::Printf( TEXT( "%s-%s" ), *HostName, *FPlatformMisc::GetLoginId().ToUpper() );
+		// If we're outside the editor and know this is the first instance, use the system login id
+		bUseStableNullId = true;
 	}
 #endif
 
+	if (bUseStableNullId)
+	{
+		// Use a stable id possibly with a user num suffix
+		return FString::Printf(TEXT( "%s-%s%s" ), *HostName, *FPlatformMisc::GetLoginId().ToUpper(), *UserSuffix);
+	}
+
 	// If we're not the first instance (or in the editor), return truly random id
-	return FString::Printf( TEXT( "%s-%s" ), *HostName, *FGuid::NewGuid().ToString() );
+	return FString::Printf(TEXT( "%s-%s%s" ), *HostName, *FGuid::NewGuid().ToString(), *UserSuffix);
 }
 
 bool FOnlineIdentityNull::Login(int32 LocalUserNum, const FOnlineAccountCredentials& AccountCredentials)
@@ -77,7 +89,7 @@ bool FOnlineIdentityNull::Login(int32 LocalUserNum, const FOnlineAccountCredenti
 	{
 		ErrorStr = FString::Printf(TEXT("Invalid LocalUserNum=%d"), LocalUserNum);
 	}
-	else if (AccountCredentials.Id.IsEmpty())
+	else if (AccountCredentials.Id.IsEmpty() && bRequireLoginCredentials)
 	{
 		ErrorStr = TEXT("Invalid account id, string empty");
 	}
@@ -156,7 +168,7 @@ bool FOnlineIdentityNull::AutoLogin(int32 LocalUserNum)
 		{
 			if (!TypeStr.IsEmpty())
 			{
-				return Login(0, FOnlineAccountCredentials(TypeStr, LoginStr, PasswordStr));
+				return Login(LocalUserNum, FOnlineAccountCredentials(TypeStr, LoginStr, PasswordStr));
 			}
 			else if (bEnableWarning)
 			{
@@ -168,6 +180,11 @@ bool FOnlineIdentityNull::AutoLogin(int32 LocalUserNum)
 			UE_LOG_ONLINE_IDENTITY(Warning, TEXT("AutoLogin missing AUTH_PASSWORD=<password>."));
 		}
 	}
+	else if (!bRequireLoginCredentials)
+	{
+		// Act like a console and login with empty auth
+		return Login(LocalUserNum, FOnlineAccountCredentials());
+	}
 	else if (bEnableWarning)
 	{
 		UE_LOG_ONLINE_IDENTITY(Warning, TEXT("AutoLogin missing AUTH_LOGIN=<login id>."));
@@ -241,6 +258,10 @@ ELoginStatus::Type FOnlineIdentityNull::GetLoginStatus(const FUniqueNetId& UserI
 	if (UserAccount.IsValid() &&
 		UserAccount->GetUserId()->IsValid())
 	{
+		if (bForceOfflineMode)
+		{
+			return ELoginStatus::UsingLocalProfile;
+		}
 		return ELoginStatus::LoggedIn;
 	}
 	return ELoginStatus::NotLoggedIn;
@@ -289,8 +310,23 @@ void FOnlineIdentityNull::RevokeAuthToken(const FUniqueNetId& UserId, const FOnR
 FOnlineIdentityNull::FOnlineIdentityNull(FOnlineSubsystemNull* InSubsystem)
 	: NullSubsystem(InSubsystem)
 {
-	// autologin the 0-th player
-	Login(0, FOnlineAccountCredentials(TEXT("DummyType"), TEXT("DummyUser"), TEXT("DummyId")) );
+	// Read configuration variables for emulating other login systems
+	GConfig->GetBool(TEXT("OnlineSubsystemNull"), TEXT("bAutoLoginAtStartup"), bAutoLoginAtStartup, GEngineIni);
+	GConfig->GetBool(TEXT("OnlineSubsystemNull"), TEXT("bRequireLoginCredentials"), bRequireLoginCredentials, GEngineIni);
+	GConfig->GetBool(TEXT("OnlineSubsystemNull"), TEXT("bAddUserNumToNullId"), bAddUserNumToNullId, GEngineIni);
+	GConfig->GetBool(TEXT("OnlineSubsystemNull"), TEXT("bForceStableNullId"), bForceStableNullId, GEngineIni);
+	GConfig->GetBool(TEXT("OnlineSubsystemNull"), TEXT("bForceOfflineMode"), bForceOfflineMode, GEngineIni);
+	
+	if (FParse::Param(FCommandLine::Get(), TEXT("StableNullID")))
+	{
+		bForceStableNullId = true;
+	}
+
+	if (bAutoLoginAtStartup)
+	{
+		// autologin the 0-th player
+		Login(0, FOnlineAccountCredentials(TEXT("DummyType"), TEXT("DummyUser"), TEXT("DummyId")) );
+	}
 }
 
 FOnlineIdentityNull::~FOnlineIdentityNull()
@@ -299,6 +335,10 @@ FOnlineIdentityNull::~FOnlineIdentityNull()
 
 void FOnlineIdentityNull::GetUserPrivilege(const FUniqueNetId& UserId, EUserPrivileges::Type Privilege, const FOnGetUserPrivilegeCompleteDelegate& Delegate)
 {
+	if (bForceOfflineMode && Privilege == EUserPrivileges::CanPlayOnline)
+	{
+		Delegate.ExecuteIfBound(UserId, Privilege, (uint32)EPrivilegeResults::NetworkConnectionUnavailable);
+	}
 	Delegate.ExecuteIfBound(UserId, Privilege, (uint32)EPrivilegeResults::NoFailures);
 }
 
diff --git a/Engine/Plugins/Online/OnlineSubsystemNull/Source/Private/OnlineIdentityNull.h b/Engine/Plugins/Online/OnlineSubsystemNull/Source/Private/OnlineIdentityNull.h
--- a/Engine/Plugins/Online/OnlineSubsystemNull/Source/Private/OnlineIdentityNull.h
+++ b/Engine/Plugins/Online/OnlineSubsystemNull/Source/Private/OnlineIdentityNull.h
@@ -92,6 +92,9 @@ public:
 	 */
 	virtual ~FOnlineIdentityNull();
 
+	/** Creates a unique id for a user, may be stable based on config/command line */
+	FString GenerateRandomUserId(int32 LocalUserNum);
+
 private:
 
 	/**
@@ -107,6 +110,21 @@ private:
 
 	/** Ids mapped to locally registered users */
 	TUniqueNetIdMap<TSharedRef<FUserOnlineAccountNull>> UserAccounts;
+
+	/** True if it should login the first user at startup like single-user platforms, false to only login when requested */
+	bool bAutoLoginAtStartup = true;
+
+	/** True if login should require a user/pass to act like an external service, false to match most platforms and use the default */
+	bool bRequireLoginCredentials = false;
+
+	/** True if login name should include the local user number, which allows different stable IDs per user num */
+	bool bAddUserNumToNullId = false;
+
+	/** True if it should use a system-stable null Id for login, same as -StableNullID on command line */
+	bool bForceStableNullId = false;
+
+	/** True if it should fail faked network queries and act like an offline system */
+	bool bForceOfflineMode = false;
 };
 
 typedef TSharedPtr<FOnlineIdentityNull, ESPMode::ThreadSafe> FOnlineIdentityNullPtr;
