Make sure Voice Talkers Are Unregistered When Destroying Session in OSS Null.
#3988

diff --git a/Engine/Plugins/Online/OnlineSubsystemNull/Source/Private/OnlineSessionInterfaceNull.cpp b/Engine/Plugins/Online/OnlineSubsystemNull/Source/Private/OnlineSessionInterfaceNull.cpp
--- a/Engine/Plugins/Online/OnlineSubsystemNull/Source/Private/OnlineSessionInterfaceNull.cpp
+++ b/Engine/Plugins/Online/OnlineSubsystemNull/Source/Private/OnlineSessionInterfaceNull.cpp
@@ -458,6 +458,22 @@ bool FOnlineSessionNull::DestroySession(FName SessionName, const FOnDestroySessi
 		UE_LOG_ONLINE_SESSION(Warning, TEXT("Can't destroy a null online session (%s)"), *SessionName.ToString());
 	}
 
+	if (GetNumSessions() == 0)
+	{
+		IOnlineVoicePtr VoiceInt = NullSubsystem->GetVoiceInterface();
+		if (VoiceInt.IsValid())
+		{
+			if (!NullSubsystem->IsDedicated())
+			{
+				// Stop local talkers
+				VoiceInt->UnregisterLocalTalkers();
+			}
+
+			// Stop remote voice 
+			VoiceInt->RemoveAllRemoteTalkers();
+		}
+	}
+
 	if (Result != ONLINE_IO_PENDING)
 	{
 		CompletionDelegate.ExecuteIfBound(SessionName, (Result == ONLINE_SUCCESS) ? true : false);
