Fix OSSTests Build

diff --git a/Engine/Plugins/Online/OnlineSubsystemNull/Source/Private/OnlineSessionInterfaceNull.cpp b/Engine/Plugins/Online/OnlineSubsystemNull/Source/Private/OnlineSessionInterfaceNull.cpp
--- a/Engine/Plugins/Online/OnlineSubsystemNull/Source/Private/OnlineSessionInterfaceNull.cpp
+++ b/Engine/Plugins/Online/OnlineSubsystemNull/Source/Private/OnlineSessionInterfaceNull.cpp
@@ -458,6 +458,7 @@ bool FOnlineSessionNull::DestroySession(FName SessionName, const FOnDestroySessi
 		UE_LOG_ONLINE_SESSION(Warning, TEXT("Can't destroy a null online session (%s)"), *SessionName.ToString());
 	}
 
+#if WITH_ENGINE
 	if (GetNumSessions() == 0)
 	{
 		IOnlineVoicePtr VoiceInt = NullSubsystem->GetVoiceInterface();
@@ -473,6 +474,7 @@ bool FOnlineSessionNull::DestroySession(FName SessionName, const FOnDestroySessi
 			VoiceInt->RemoveAllRemoteTalkers();
 		}
 	}
+#endif //WITH_ENGINE
 
 	if (Result != ONLINE_IO_PENDING)
 	{
