Added an overload to the Checkout method of the IOnlinePurchase interface that don't return receipts. Recommended to use this overload if your product processes purchases on backend services, and the clients don't need to know about the purchase.

diff --git a/Engine/Plugins/Online/OnlineSubsystem/Source/Public/Interfaces/OnlinePurchaseInterface.h b/Engine/Plugins/Online/OnlineSubsystem/Source/Public/Interfaces/OnlinePurchaseInterface.h
--- a/Engine/Plugins/Online/OnlineSubsystem/Source/Public/Interfaces/OnlinePurchaseInterface.h
+++ b/Engine/Plugins/Online/OnlineSubsystem/Source/Public/Interfaces/OnlinePurchaseInterface.h
@@ -212,6 +212,11 @@ public:
  */
 DECLARE_DELEGATE_TwoParams(FOnPurchaseCheckoutComplete, const FOnlineError& /*Result*/, const TSharedRef<FPurchaseReceipt>& /*Receipt*/);
 
+/**
+* Delegate called when checkout process completes, this delegate is used when the entitlement or receipt information of the purchase is not needed by the caller.
+*/
+DECLARE_DELEGATE_OneParam(FOnPurchaseReceiptlessCheckoutComplete, const FOnlineError& /*Result*/);
+
 /**
  * Delegate called when code redemption process completes
  */
@@ -258,6 +263,17 @@ public:
 	 * @param Delegate completion callback (guaranteed to be called)
 	 */
 	virtual void Checkout(const FUniqueNetId& UserId, const FPurchaseCheckoutRequest& CheckoutRequest, const FOnPurchaseCheckoutComplete& Delegate) = 0;
+
+
+	/**
+	 * Initiate the checkout process for purchasing offers via payment without receipts returned
+	 * Recommended for use when the entitlement or receipt information of the purchase is not needed by the caller
+	 * 
+	 * @param UserId user initiating the request
+	 * @param CheckoutRequest info needed for the checkout request
+	 * @param Delegate completion callback (guaranteed to be called)
+	 */
+	virtual void Checkout(const FUniqueNetId& UserId, const FPurchaseCheckoutRequest& CheckoutRequest, const FOnPurchaseReceiptlessCheckoutComplete& Delegate) = 0;
 	
 	/**
 	 * Finalizes a purchase with the supporting platform
@@ -319,6 +335,9 @@ public:
 	 *
 	 */
 	DEFINE_ONLINE_DELEGATE_ONE_PARAM(OnUnexpectedPurchaseReceipt, const FUniqueNetId& /*UserId*/);
+
+private:
+	void OnRedirectToCheckoutComplete(const FOnlineError& ErrorResult, const TSharedRef<FPurchaseReceipt>& Receipt, const FOnPurchaseReceiptlessCheckoutComplete Delegate);
 };
 
 inline const TCHAR* LexToString(EPurchaseTransactionState State)
diff --git a/Engine/Plugins/Online/OnlineSubsystemEOS/Source/OnlineSubsystemEOS/Private/OnlineStoreEOS.cpp b/Engine/Plugins/Online/OnlineSubsystemEOS/Source/OnlineSubsystemEOS/Private/OnlineStoreEOS.cpp
--- a/Engine/Plugins/Online/OnlineSubsystemEOS/Source/OnlineSubsystemEOS/Private/OnlineStoreEOS.cpp
+++ b/Engine/Plugins/Online/OnlineSubsystemEOS/Source/OnlineSubsystemEOS/Private/OnlineStoreEOS.cpp
@@ -221,7 +221,12 @@ void FOnlineStoreEOS::Checkout(const FUniqueNetId& UserId, const FPurchaseChecko
 
 	};
 	EOS_Ecom_Checkout(EOSSubsystem->EcomHandle, &Options, CallbackObj, CallbackObj->GetCallbackPtr());
+}
 
+void FOnlineStoreEOS::Checkout(const FUniqueNetId& UserId, const FPurchaseCheckoutRequest& CheckoutRequest, const FOnPurchaseReceiptlessCheckoutComplete& Delegate)
+{
+	// Checkout with no receipt query Delegate is not implemented, please use the other Checkout method
+	Delegate.ExecuteIfBound(ONLINE_ERROR(EOnlineErrorResult::NotImplemented));
 }
 
 void FOnlineStoreEOS::FinalizePurchase(const FUniqueNetId& UserId, const FString& ReceiptId)
diff --git a/Engine/Plugins/Online/OnlineSubsystemEOS/Source/OnlineSubsystemEOS/Private/OnlineStoreEOS.h b/Engine/Plugins/Online/OnlineSubsystemEOS/Source/OnlineSubsystemEOS/Private/OnlineStoreEOS.h
--- a/Engine/Plugins/Online/OnlineSubsystemEOS/Source/OnlineSubsystemEOS/Private/OnlineStoreEOS.h
+++ b/Engine/Plugins/Online/OnlineSubsystemEOS/Source/OnlineSubsystemEOS/Private/OnlineStoreEOS.h
@@ -36,6 +36,7 @@ public:
 // Begin IOnlinePurchase
 	virtual bool IsAllowedToPurchase(const FUniqueNetId& UserId) override { return true; }
 	virtual void Checkout(const FUniqueNetId& UserId, const FPurchaseCheckoutRequest& CheckoutRequest, const FOnPurchaseCheckoutComplete& Delegate) override;
+	virtual void Checkout(const FUniqueNetId& UserId, const FPurchaseCheckoutRequest& CheckoutRequest, const FOnPurchaseReceiptlessCheckoutComplete& Delegate) override;
 	virtual void FinalizePurchase(const FUniqueNetId& UserId, const FString& ReceiptId) override;
 	virtual void RedeemCode(const FUniqueNetId& UserId, const FRedeemCodeRequest& RedeemCodeRequest, const FOnPurchaseRedeemCodeComplete& Delegate) override;
 	virtual void QueryReceipts(const FUniqueNetId& UserId, bool bRestoreReceipts, const FOnQueryReceiptsComplete& Delegate) override;
diff --git a/Engine/Plugins/Online/OnlineSubsystemEOS/Source/OnlineSubsystemEOSPlus/Private/OnlineStoreEOSPlus.cpp b/Engine/Plugins/Online/OnlineSubsystemEOS/Source/OnlineSubsystemEOSPlus/Private/OnlineStoreEOSPlus.cpp
--- a/Engine/Plugins/Online/OnlineSubsystemEOS/Source/OnlineSubsystemEOSPlus/Private/OnlineStoreEOSPlus.cpp
+++ b/Engine/Plugins/Online/OnlineSubsystemEOS/Source/OnlineSubsystemEOSPlus/Private/OnlineStoreEOSPlus.cpp
@@ -215,6 +215,7 @@ void FOnlineStoreEOSPlus::Checkout(const FUniqueNetId& UserId, const FPurchaseCh
 {
 	EOnlineErrorResult ErrorResult = EOnlineErrorResult::Unknown;
 	FUniqueNetIdEOSPlusPtr NetIdPlus = GetNetIdPlus(UserId.ToString());
+
 	if (NetIdPlus.IsValid())
 	{
 		const bool bIsBaseNetIdValid = ensure(NetIdPlus->GetBaseNetId().IsValid());
@@ -242,6 +243,36 @@ void FOnlineStoreEOSPlus::Checkout(const FUniqueNetId& UserId, const FPurchaseCh
 	});
 }
 
+void FOnlineStoreEOSPlus::Checkout(const FUniqueNetId& UserId, const FPurchaseCheckoutRequest& CheckoutRequest, const FOnPurchaseReceiptlessCheckoutComplete& Delegate)
+{
+	EOnlineErrorResult ErrorResult = EOnlineErrorResult::Unknown;
+	FUniqueNetIdEOSPlusPtr NetIdPlus = GetNetIdPlus(UserId.ToString());
+	if (NetIdPlus.IsValid())
+	{
+		const bool bIsBaseNetIdValid = ensure(NetIdPlus->GetBaseNetId().IsValid());
+		const bool bIsBasePurchaseInterfaceValid = BasePurchaseInterface.IsValid();
+		if (bIsBaseNetIdValid && bIsBasePurchaseInterfaceValid)
+		{
+			BasePurchaseInterface->Checkout(*NetIdPlus->GetBaseNetId(), CheckoutRequest, Delegate);
+			return;
+		}
+		else
+		{
+			UE_LOG_ONLINE(Warning, TEXT("[FOnlineStoreEOSPlus::Checkout] Unable to call method in base interface. IsBaseNetIdValid=%s IsBasePurchaseInterfaceValid=%s."), *LexToString(bIsBaseNetIdValid), *LexToString(bIsBasePurchaseInterfaceValid));
+			ErrorResult = bIsBasePurchaseInterfaceValid ? EOnlineErrorResult::MissingInterface : EOnlineErrorResult::InvalidUser;
+		}
+	}
+	else
+	{
+		UE_LOG_ONLINE(Warning, TEXT("[FOnlineStoreEOSPlus::Checkout] Unable to call method in base interface. Unknown user (%s)"), *UserId.ToString());
+		ErrorResult = EOnlineErrorResult::InvalidUser;
+	}
+
+	EOSPlus->ExecuteNextTick([this, ErrorResult, Delegate]() {
+		Delegate.ExecuteIfBound(FOnlineError(ErrorResult));
+	});
+}
+
 void FOnlineStoreEOSPlus::FinalizePurchase(const FUniqueNetId& UserId, const FString& ReceiptId)
 {
 	FUniqueNetIdEOSPlusPtr NetIdPlus = GetNetIdPlus(UserId.ToString());
diff --git a/Engine/Plugins/Online/OnlineSubsystemEOS/Source/OnlineSubsystemEOSPlus/Private/OnlineStoreEOSPlus.h b/Engine/Plugins/Online/OnlineSubsystemEOS/Source/OnlineSubsystemEOSPlus/Private/OnlineStoreEOSPlus.h
--- a/Engine/Plugins/Online/OnlineSubsystemEOS/Source/OnlineSubsystemEOSPlus/Private/OnlineStoreEOSPlus.h
+++ b/Engine/Plugins/Online/OnlineSubsystemEOS/Source/OnlineSubsystemEOSPlus/Private/OnlineStoreEOSPlus.h
@@ -37,6 +37,7 @@ public:
 	//~ Begin IOnlinePurchase Interface
 	virtual bool IsAllowedToPurchase(const FUniqueNetId& UserId) override;
 	virtual void Checkout(const FUniqueNetId& UserId, const FPurchaseCheckoutRequest& CheckoutRequest, const FOnPurchaseCheckoutComplete& Delegate) override;
+	virtual void Checkout(const FUniqueNetId& UserId, const FPurchaseCheckoutRequest& CheckoutRequest, const FOnPurchaseReceiptlessCheckoutComplete& Delegate) override;
 	virtual void FinalizePurchase(const FUniqueNetId& UserId, const FString& ReceiptId) override;
 	virtual void RedeemCode(const FUniqueNetId& UserId, const FRedeemCodeRequest& RedeemCodeRequest, const FOnPurchaseRedeemCodeComplete& Delegate) override;
 	virtual void QueryReceipts(const FUniqueNetId& UserId, bool bRestoreReceipts, const FOnQueryReceiptsComplete& Delegate) override;
diff --git a/Engine/Plugins/Online/OnlineSubsystemNull/Source/Private/OnlinePurchaseInterfaceNull.cpp b/Engine/Plugins/Online/OnlineSubsystemNull/Source/Private/OnlinePurchaseInterfaceNull.cpp
--- a/Engine/Plugins/Online/OnlineSubsystemNull/Source/Private/OnlinePurchaseInterfaceNull.cpp
+++ b/Engine/Plugins/Online/OnlineSubsystemNull/Source/Private/OnlinePurchaseInterfaceNull.cpp
@@ -141,6 +141,85 @@ void FOnlinePurchaseNull::CheckoutSuccessfully(const FUniqueNetIdNull& UserId, T
 	Delegate.ExecuteIfBound(FOnlineError(true), MakeShared<FPurchaseReceipt>(PurchaseReceipt));
 }
 
+void FOnlinePurchaseNull::Checkout(const FUniqueNetId& UserId, const FPurchaseCheckoutRequest& CheckoutRequest, const FOnPurchaseReceiptlessCheckoutComplete& Delegate)
+{
+	// Lambda to wrap calling our delegate with an error and logging the message
+	auto CallDelegateError = [this, &Delegate](FString&& ErrorMessage)
+	{
+		NullSubsystem.ExecuteNextTick([Delegate, MovedErrorMessage = MoveTemp(ErrorMessage)]() mutable
+		{
+			UE_LOG_ONLINE(Error, TEXT("%s"), *MovedErrorMessage);
+
+			Delegate.ExecuteIfBound(FOnlineError(MoveTemp(MovedErrorMessage)));
+		});
+	};
+
+	if (CheckoutRequest.PurchaseOffers.Num() == 0)
+	{
+		CallDelegateError(TEXT("FOnlinePurchaseNull::Checkout failed, there were no entries passed to purchase"));
+		return;
+	}
+	else if (CheckoutRequest.PurchaseOffers.Num() != 1)
+	{
+		CallDelegateError(TEXT("FOnlinePurchaseNull::Checkout failed, there were more than one entry passed to purchase. We currently only support one."));
+		return;
+	}
+
+	check(CheckoutRequest.PurchaseOffers.IsValidIndex(0));
+	const FPurchaseCheckoutRequest::FPurchaseOfferEntry& Entry = CheckoutRequest.PurchaseOffers[0];
+
+	if (Entry.Quantity != 1)
+	{
+		CallDelegateError(TEXT("FOnlinePurchaseNull::Checkout failed, purchase quantity not set to one. We currently only support one."));
+		return;
+	}
+
+	if (Entry.OfferId.IsEmpty())
+	{
+		CallDelegateError(TEXT("FOnlinePurchaseNull::Checkout failed, OfferId is blank."));
+		return;
+	}
+
+	const IOnlineStoreV2Ptr NullStoreInt = NullSubsystem.GetStoreV2Interface();
+
+	TSharedPtr<FOnlineStoreOffer> NullOffer = NullStoreInt->GetOffer(Entry.OfferId);
+	if (!NullOffer.IsValid())
+	{
+		CallDelegateError(TEXT("FOnlinePurchaseNull::Checkout failed, Could not find corresponding offer."));
+		return;
+	}
+
+	if (PendingPurchaseHandle.IsValid())
+	{
+		CallDelegateError(TEXT("FOnlinePurchaseNull::Checkout failed, there was another purchase in progress."));
+		return;
+	}
+
+	PendingPurchaseHandle = Delegate.GetHandle();
+
+	TWeakPtr<FOnlinePurchaseNull, ESPMode::ThreadSafe> WeakMe = AsShared();
+	const FUniqueNetIdNull& NullUserId = static_cast<const FUniqueNetIdNull&>(UserId);
+
+	NullSubsystem.ExecuteNextTick([NullUserId, WeakMe, Delegate]
+		{
+			FOnlinePurchaseNullPtr StrongThis = WeakMe.Pin();
+			if (StrongThis.IsValid())
+			{
+				StrongThis->CheckoutSuccessfully(NullUserId, Delegate);
+			}
+		});
+}
+
+void FOnlinePurchaseNull::CheckoutSuccessfully(const FUniqueNetIdNull& UserId, FOnPurchaseReceiptlessCheckoutComplete Delegate)
+{
+	// Reset pending purchase handle
+	check(PendingPurchaseHandle.IsValid());
+	PendingPurchaseHandle.Reset();
+	
+	// Finish pending purchase
+	Delegate.ExecuteIfBound(FOnlineError(EOnlineErrorResult::Success));
+}
+
 void FOnlinePurchaseNull::FinalizePurchase(const FUniqueNetId& UserId, const FString& ReceiptId)
 {
 	const FUniqueNetIdNull& NullUserId = static_cast<const FUniqueNetIdNull&>(UserId);
diff --git a/Engine/Plugins/Online/OnlineSubsystemNull/Source/Private/OnlinePurchaseInterfaceNull.h b/Engine/Plugins/Online/OnlineSubsystemNull/Source/Private/OnlinePurchaseInterfaceNull.h
--- a/Engine/Plugins/Online/OnlineSubsystemNull/Source/Private/OnlinePurchaseInterfaceNull.h
+++ b/Engine/Plugins/Online/OnlineSubsystemNull/Source/Private/OnlinePurchaseInterfaceNull.h
@@ -21,6 +21,7 @@ public:
 	//~ Begin IOnlinePurchase Interface
 	virtual bool IsAllowedToPurchase(const FUniqueNetId& UserId) override;
 	virtual void Checkout(const FUniqueNetId& UserId, const FPurchaseCheckoutRequest& CheckoutRequest, const FOnPurchaseCheckoutComplete& Delegate) override;
+	virtual void Checkout(const FUniqueNetId& UserId, const FPurchaseCheckoutRequest& CheckoutRequest, const FOnPurchaseReceiptlessCheckoutComplete& Delegate) override;
 	virtual void FinalizePurchase(const FUniqueNetId& UserId, const FString& ReceiptId) override;
 	virtual void RedeemCode(const FUniqueNetId& UserId, const FRedeemCodeRequest& RedeemCodeRequest, const FOnPurchaseRedeemCodeComplete& Delegate) override;
 	virtual void QueryReceipts(const FUniqueNetId& UserId, bool bRestoreReceipts, const FOnQueryReceiptsComplete& Delegate) override;
@@ -30,6 +31,7 @@ public:
 
 PACKAGE_SCOPE:
 	void CheckoutSuccessfully(const FUniqueNetIdNull& UserId, TSharedPtr<FOnlineStoreOffer> Offer);
+	void CheckoutSuccessfully(const FUniqueNetIdNull& UserId, FOnPurchaseReceiptlessCheckoutComplete Delegate);
 
 PACKAGE_SCOPE:
 	/** Pointer back to our parent subsystem */
@@ -42,6 +44,8 @@ PACKAGE_SCOPE:
 	TOptional<FOnPurchaseCheckoutComplete> PendingPurchaseDelegate;
 
 	TOptional<double> PendingPurchaseFailTime;
+
+	FDelegateHandle PendingPurchaseHandle;
 };
 
 using FOnlinePurchaseNullPtr = TSharedPtr<FOnlinePurchaseNull, ESPMode::ThreadSafe>;
