Add OnlineExternalUINull (disabled by default) OSS interface to allow emulating console/mobile platforms on PC
Move existing OSS Null testing config variables to OnlineSubsystemNull and expose them as proper cvars so they can be changed between PIE runs
Add SupportExternalUI, RequireShowLoginUI, and ForceShowLoginUIUserChange cvars/config options to emulate a platform user change during login

diff --git a/Engine/Plugins/Online/OnlineSubsystemNull/Source/Private/OnlineExternalUINull.cpp b/Engine/Plugins/Online/OnlineSubsystemNull/Source/Private/OnlineExternalUINull.cpp
new file mode 100644
--- /dev/null
+++ b/Engine/Plugins/Online/OnlineSubsystemNull/Source/Private/OnlineExternalUINull.cpp
@@ -0,0 +1,103 @@
+// Copyright Epic Games, Inc. All Rights Reserved.
+
+#include "OnlineExternalUINull.h"
+#include "OnlineSubsystemNull.h"
+#include "OnlineIdentityNull.h"
+
+FOnlineExternalUINull::FOnlineExternalUINull(FOnlineSubsystemNull* InSubsystem)
+	: NullSubsystem(InSubsystem)
+{
+
+}
+
+FOnlineExternalUINull::~FOnlineExternalUINull()
+{
+
+}
+
+bool FOnlineExternalUINull::ShowLoginUI(const int ControllerIndex, bool bShowOnlineOnly, bool bShowSkipButton, const FOnLoginUIClosedDelegate& Delegate)
+{
+	NullSubsystem->ExecuteNextTick([this, ControllerIndex, Delegate]()
+	{
+		// Just call identity login, possibly with a different user than initially
+		int32 NewUserIndex = ControllerIndex;
+		if (FOnlineSubsystemNull::bForceShowLoginUIUserChange)
+		{
+			NewUserIndex++;
+		}
+
+		FOnlineIdentityNullPtr IdentityInt = StaticCastSharedPtr<FOnlineIdentityNull>(NullSubsystem->GetIdentityInterface());
+		if (IdentityInt.IsValid())
+		{
+			FOnLoginCompleteDelegate CompletionDelegate;
+			CompletionDelegate = FOnLoginCompleteDelegate::CreateRaw(this, &FOnlineExternalUINull::OnIdentityLoginComplete, Delegate);
+			IdentityInt->LoginInternal(NewUserIndex, FOnlineAccountCredentials(TEXT("ShowLoginUI"), TEXT("DummyUser"), TEXT("DummyId")), CompletionDelegate);
+		}
+	});
+
+	return true;
+}
+
+void FOnlineExternalUINull::OnIdentityLoginComplete(int32 LocalUserNum, bool bWasSuccessful, const FUniqueNetId& UserId, const FString& Error, FOnLoginUIClosedDelegate Delegate)
+{
+	FUniqueNetIdPtr StrongUserId = UserId.AsShared();
+	NullSubsystem->ExecuteNextTick([StrongUserId, LocalUserNum, bWasSuccessful, Delegate]()
+	{
+		Delegate.ExecuteIfBound(StrongUserId, LocalUserNum, FOnlineError(bWasSuccessful));
+	});
+}
+
+bool FOnlineExternalUINull::ShowAccountCreationUI(const int ControllerIndex, const FOnAccountCreationUIClosedDelegate& Delegate)
+{
+	return false;
+}
+
+bool FOnlineExternalUINull::ShowFriendsUI(int32 LocalUserNum)
+{
+	return false;
+}
+
+bool FOnlineExternalUINull::ShowInviteUI(int32 LocalUserNum, FName SessionName)
+{
+	return false;
+}
+
+bool FOnlineExternalUINull::ShowAchievementsUI(int32 LocalUserNum)
+{
+	return false;
+}
+
+bool FOnlineExternalUINull::ShowLeaderboardUI( const FString& LeaderboardName )
+{
+	return false;
+}
+
+bool FOnlineExternalUINull::ShowWebURL(const FString& Url, const FShowWebUrlParams& ShowParams, const FOnShowWebUrlClosedDelegate& Delegate)
+{
+	return false;
+}
+
+bool FOnlineExternalUINull::CloseWebURL()
+{
+	return false;
+}
+
+bool FOnlineExternalUINull::ShowAccountUpgradeUI(const FUniqueNetId& UniqueId)
+{
+	return false;
+}
+
+bool FOnlineExternalUINull::ShowStoreUI(int32 LocalUserNum, const FShowStoreParams& ShowParams, const FOnShowStoreUIClosedDelegate& Delegate)
+{
+	return false;
+}
+
+bool FOnlineExternalUINull::ShowSendMessageUI(int32 LocalUserNum, const FShowSendMessageParams& ShowParams, const FOnShowSendMessageUIClosedDelegate& Delegate)
+{
+	return false;
+}
+
+bool FOnlineExternalUINull::ShowProfileUI(const FUniqueNetId& Requestor, const FUniqueNetId& Requestee, const FOnProfileUIClosedDelegate& Delegate)
+{
+	return false;
+}
diff --git a/Engine/Plugins/Online/OnlineSubsystemNull/Source/Private/OnlineExternalUINull.h b/Engine/Plugins/Online/OnlineSubsystemNull/Source/Private/OnlineExternalUINull.h
new file mode 100644
--- /dev/null
+++ b/Engine/Plugins/Online/OnlineSubsystemNull/Source/Private/OnlineExternalUINull.h
@@ -0,0 +1,54 @@
+// Copyright Epic Games, Inc. All Rights Reserved.
+
+#pragma once
+
+#include "CoreMinimal.h"
+#include "Online/CoreOnline.h"
+#include "OnlineSubsystemTypes.h"
+#include "Interfaces/OnlineExternalUIInterface.h"
+#include "OnlineSubsystemNullTypes.h"
+ 
+class FOnlineSubsystemNull;
+
+/** 
+ * Null implementation to allow testing
+ */
+class FOnlineExternalUINull : public IOnlineExternalUI
+{
+public:
+
+	/**
+	 * Constructor
+	 * @param InSubsystem The owner of this external UI interface.
+	 */
+	FOnlineExternalUINull(FOnlineSubsystemNull* InSubsystem);
+
+	/**
+	 * Destructor.
+	 */
+	virtual ~FOnlineExternalUINull();
+
+	// IOnlineExternalUI
+	virtual bool ShowLoginUI(const int ControllerIndex, bool bShowOnlineOnly, bool bShowSkipButton, const FOnLoginUIClosedDelegate& Delegate = FOnLoginUIClosedDelegate()) override;
+	virtual bool ShowAccountCreationUI(const int ControllerIndex, const FOnAccountCreationUIClosedDelegate& Delegate = FOnAccountCreationUIClosedDelegate()) override;
+	virtual bool ShowFriendsUI(int32 LocalUserNum) override;
+	virtual bool ShowInviteUI(int32 LocalUserNum, FName SessionName = NAME_GameSession) override;
+	virtual bool ShowAchievementsUI(int32 LocalUserNum) override;
+	virtual bool ShowLeaderboardUI(const FString& LeaderboardName) override;
+	virtual bool ShowWebURL(const FString& Url, const FShowWebUrlParams& ShowParams, const FOnShowWebUrlClosedDelegate& Delegate = FOnShowWebUrlClosedDelegate()) override;
+	virtual bool CloseWebURL() override;
+	virtual bool ShowProfileUI(const FUniqueNetId& Requestor, const FUniqueNetId& Requestee, const FOnProfileUIClosedDelegate& Delegate = FOnProfileUIClosedDelegate()) override;
+	virtual bool ShowAccountUpgradeUI(const FUniqueNetId& UniqueId) override;
+	virtual bool ShowStoreUI(int32 LocalUserNum, const FShowStoreParams& ShowParams, const FOnShowStoreUIClosedDelegate& Delegate = FOnShowStoreUIClosedDelegate()) override;
+	virtual bool ShowSendMessageUI(int32 LocalUserNum, const FShowSendMessageParams& ShowParams, const FOnShowSendMessageUIClosedDelegate& Delegate = FOnShowSendMessageUIClosedDelegate()) override;
+
+private:
+
+	/** Reference to the owning subsystem */
+	FOnlineSubsystemNull *NullSubsystem;
+
+	void OnIdentityLoginComplete(int32 LocalUserNum, bool bWasSuccessful, const FUniqueNetId& UserId, const FString& Error, FOnLoginUIClosedDelegate Delegate);
+};
+
+typedef TSharedPtr<FOnlineExternalUINull, ESPMode::ThreadSafe> FOnlineExternalUINullPtr;
+
diff --git a/Engine/Plugins/Online/OnlineSubsystemNull/Source/Private/OnlineIdentityNull.cpp b/Engine/Plugins/Online/OnlineSubsystemNull/Source/Private/OnlineIdentityNull.cpp
--- a/Engine/Plugins/Online/OnlineSubsystemNull/Source/Private/OnlineIdentityNull.cpp
+++ b/Engine/Plugins/Online/OnlineSubsystemNull/Source/Private/OnlineIdentityNull.cpp
@@ -54,10 +54,10 @@ FString FOnlineIdentityNull::GenerateRandomUserId(int32 LocalUserNum)
 		HostName = Addr->ToString(false);
 	}
 
-	bool bUseStableNullId = bForceStableNullId;
+	bool bUseStableNullId = FOnlineSubsystemNull::bForceStableNullId;
 	FString UserSuffix;
 
-	if (bAddUserNumToNullId)
+	if (FOnlineSubsystemNull::bAddUserNumToNullId)
 	{
 		UserSuffix = FString::Printf(TEXT("-%d"), LocalUserNum);
 	}
@@ -79,6 +79,15 @@ FString FOnlineIdentityNull::GenerateRandomUserId(int32 LocalUserNum)
 }
 
 bool FOnlineIdentityNull::Login(int32 LocalUserNum, const FOnlineAccountCredentials& AccountCredentials)
+{
+	return LoginInternal(LocalUserNum, AccountCredentials, FOnLoginCompleteDelegate::CreateLambda(
+		[this](int32 LocalUserNum, bool bWasSuccessful, const FUniqueNetId& UserId, const FString& ErrorStr)
+		{
+			TriggerOnLoginCompleteDelegates(LocalUserNum, bWasSuccessful, UserId, ErrorStr);
+		}));
+}
+
+bool FOnlineIdentityNull::LoginInternal(int32 LocalUserNum, const FOnlineAccountCredentials& AccountCredentials, const FOnLoginCompleteDelegate& InCompletionDelegate)
 {
 	FString ErrorStr;
 	TSharedPtr<FUserOnlineAccountNull> UserAccountPtr;
@@ -88,14 +97,18 @@ bool FOnlineIdentityNull::Login(int32 LocalUserNum, const FOnlineAccountCredenti
 	{
 		ErrorStr = FString::Printf(TEXT("Invalid LocalUserNum=%d"), LocalUserNum);
 	}
-	else if (AccountCredentials.Id.IsEmpty() && bRequireLoginCredentials)
+	else if (FOnlineSubsystemNull::bRequireShowLoginUI && !AccountCredentials.Type.Contains(TEXT("LoginUI")))
+	{
+		ErrorStr = TEXT("Only supports login using ShowLoginUI");
+	}
+	else if (FOnlineSubsystemNull::bRequireLoginCredentials && AccountCredentials.Id.IsEmpty())
 	{
 		ErrorStr = TEXT("Invalid account id, string empty");
 	}
 	else
 	{
 		FUniqueNetIdPtr* UserId = UserIds.Find(LocalUserNum);
-		if (UserId == NULL)
+		if (UserId == nullptr)
 		{
 			FString RandomUserId = GenerateRandomUserId(LocalUserNum);
 
@@ -118,11 +131,11 @@ bool FOnlineIdentityNull::Login(int32 LocalUserNum, const FOnlineAccountCredenti
 	if (!ErrorStr.IsEmpty())
 	{
 		UE_LOG_ONLINE_IDENTITY(Warning, TEXT("Login request failed. %s"), *ErrorStr);
-		TriggerOnLoginCompleteDelegates(LocalUserNum, false, *FUniqueNetIdNull::EmptyId(), ErrorStr);
+		InCompletionDelegate.Execute(LocalUserNum, false, *FUniqueNetIdNull::EmptyId(), ErrorStr);
 		return false;
 	}
 
-	TriggerOnLoginCompleteDelegates(LocalUserNum, true, *UserAccountPtr->GetUserId(), ErrorStr);
+	InCompletionDelegate.Execute(LocalUserNum, true, *UserAccountPtr->GetUserId(), ErrorStr);
 	return true;
 }
 
@@ -179,7 +192,7 @@ bool FOnlineIdentityNull::AutoLogin(int32 LocalUserNum)
 			UE_LOG_ONLINE_IDENTITY(Warning, TEXT("AutoLogin missing AUTH_PASSWORD=<password>."));
 		}
 	}
-	else if (!bRequireLoginCredentials)
+	else if (!FOnlineSubsystemNull::bRequireLoginCredentials)
 	{
 		// Act like a console and login with empty auth
 		return Login(LocalUserNum, FOnlineAccountCredentials());
@@ -257,7 +270,7 @@ ELoginStatus::Type FOnlineIdentityNull::GetLoginStatus(const FUniqueNetId& UserI
 	if (UserAccount.IsValid() &&
 		UserAccount->GetUserId()->IsValid())
 	{
-		if (bForceOfflineMode)
+		if (FOnlineSubsystemNull::bForceOfflineMode)
 		{
 			return ELoginStatus::UsingLocalProfile;
 		}
@@ -308,20 +321,8 @@ void FOnlineIdentityNull::RevokeAuthToken(const FUniqueNetId& UserId, const FOnR
 
 FOnlineIdentityNull::FOnlineIdentityNull(FOnlineSubsystemNull* InSubsystem)
 	: NullSubsystem(InSubsystem)
-{
-	// Read configuration variables for emulating other login systems
-	GConfig->GetBool(TEXT("OnlineSubsystemNull"), TEXT("bAutoLoginAtStartup"), bAutoLoginAtStartup, GEngineIni);
-	GConfig->GetBool(TEXT("OnlineSubsystemNull"), TEXT("bRequireLoginCredentials"), bRequireLoginCredentials, GEngineIni);
-	GConfig->GetBool(TEXT("OnlineSubsystemNull"), TEXT("bAddUserNumToNullId"), bAddUserNumToNullId, GEngineIni);
-	GConfig->GetBool(TEXT("OnlineSubsystemNull"), TEXT("bForceStableNullId"), bForceStableNullId, GEngineIni);
-	GConfig->GetBool(TEXT("OnlineSubsystemNull"), TEXT("bForceOfflineMode"), bForceOfflineMode, GEngineIni);
-	
-	if (FParse::Param(FCommandLine::Get(), TEXT("StableNullID")))
-	{
-		bForceStableNullId = true;
-	}
-
-	if (bAutoLoginAtStartup)
+{	
+	if (FOnlineSubsystemNull::bAutoLoginAtStartup)
 	{
 		// autologin the 0-th player
 		Login(0, FOnlineAccountCredentials(TEXT("DummyType"), TEXT("DummyUser"), TEXT("DummyId")) );
@@ -334,7 +335,7 @@ FOnlineIdentityNull::~FOnlineIdentityNull()
 
 void FOnlineIdentityNull::GetUserPrivilege(const FUniqueNetId& UserId, EUserPrivileges::Type Privilege, const FOnGetUserPrivilegeCompleteDelegate& Delegate)
 {
-	if (bForceOfflineMode && Privilege == EUserPrivileges::CanPlayOnline)
+	if (FOnlineSubsystemNull::bForceOfflineMode && Privilege == EUserPrivileges::CanPlayOnline)
 	{
 		Delegate.ExecuteIfBound(UserId, Privilege, (uint32)EPrivilegeResults::NetworkConnectionUnavailable);
 	}
diff --git a/Engine/Plugins/Online/OnlineSubsystemNull/Source/Private/OnlineIdentityNull.h b/Engine/Plugins/Online/OnlineSubsystemNull/Source/Private/OnlineIdentityNull.h
--- a/Engine/Plugins/Online/OnlineSubsystemNull/Source/Private/OnlineIdentityNull.h
+++ b/Engine/Plugins/Online/OnlineSubsystemNull/Source/Private/OnlineIdentityNull.h
@@ -92,6 +92,9 @@ public:
 	 */
 	virtual ~FOnlineIdentityNull();
 
+	/** Login and call arbitrary callback instead of registered one */
+	bool LoginInternal(int32 LocalUserNum, const FOnlineAccountCredentials& AccountCredentials, const FOnLoginCompleteDelegate& InCompletionDelegate);
+
 	/** Creates a unique id for a user, may be stable based on config/command line */
 	FString GenerateRandomUserId(int32 LocalUserNum);
 
@@ -110,21 +113,6 @@ private:
 
 	/** Ids mapped to locally registered users */
 	TUniqueNetIdMap<TSharedRef<FUserOnlineAccountNull>> UserAccounts;
-
-	/** True if it should login the first user at startup like single-user platforms, false to only login when requested */
-	bool bAutoLoginAtStartup = true;
-
-	/** True if login should require a user/pass to act like an external service, false to match most platforms and use the default */
-	bool bRequireLoginCredentials = false;
-
-	/** True if login name should include the local user number, which allows different stable IDs per user num */
-	bool bAddUserNumToNullId = false;
-
-	/** True if it should use a system-stable null Id for login, same as -StableNullID on command line */
-	bool bForceStableNullId = false;
-
-	/** True if it should fail faked network queries and act like an offline system */
-	bool bForceOfflineMode = false;
 };
 
 typedef TSharedPtr<FOnlineIdentityNull, ESPMode::ThreadSafe> FOnlineIdentityNullPtr;
diff --git a/Engine/Plugins/Online/OnlineSubsystemNull/Source/Private/OnlineSubsystemNull.cpp b/Engine/Plugins/Online/OnlineSubsystemNull/Source/Private/OnlineSubsystemNull.cpp
--- a/Engine/Plugins/Online/OnlineSubsystemNull/Source/Private/OnlineSubsystemNull.cpp
+++ b/Engine/Plugins/Online/OnlineSubsystemNull/Source/Private/OnlineSubsystemNull.cpp
@@ -7,6 +7,7 @@
 #include "OnlineSessionInterfaceNull.h"
 #include "OnlineLeaderboardInterfaceNull.h"
 #include "OnlineIdentityNull.h"
+#include "OnlineExternalUINull.h"
 #include "OnlineAchievementsInterfaceNull.h"
 #include "OnlineStoreV2InterfaceNull.h"
 #include "OnlinePurchaseInterfaceNull.h"
@@ -81,7 +82,8 @@ IOnlineVoicePtr FOnlineSubsystemNull::GetVoiceInterface() const
 
 IOnlineExternalUIPtr FOnlineSubsystemNull::GetExternalUIInterface() const
 {
-	return nullptr;
+	// This may be null depending on cvars
+	return ExternalUIInterface;
 }
 
 IOnlineTimePtr FOnlineSubsystemNull::GetTimeInterface() const
@@ -193,12 +195,91 @@ bool FOnlineSubsystemNull::Tick(float DeltaTime)
 	return true;
 }
 
+// Declares console variables that modify how OSS null works so it can emulate different platforms for testing. These are also loaded from [OnlineSubsystemNull]
+
+bool FOnlineSubsystemNull::bAutoLoginAtStartup = true;
+FAutoConsoleVariableRef CVarAutoLoginAtStartup(
+	TEXT("OSSNull.AutoLoginAtStartup"),
+	FOnlineSubsystemNull::bAutoLoginAtStartup,
+	TEXT("True if it should login the first user at startup like single-user platforms, false to only login when requested"),
+	ECVF_Default | ECVF_Cheat);
+
+bool FOnlineSubsystemNull::bSupportExternalUI = false;
+FAutoConsoleVariableRef CVarSupportExternalUI(
+	TEXT("OSSNull.SupportExternalUI"),
+	FOnlineSubsystemNull::bSupportExternalUI,
+	TEXT("True if it should support an external UI interface"),
+	ECVF_Default | ECVF_Cheat);
+
+bool FOnlineSubsystemNull::bRequireShowLoginUI = false;
+FAutoConsoleVariableRef CVarRequireShowLoginUI(
+	TEXT("OSSNull.RequireShowLoginUI"),
+	FOnlineSubsystemNull::bRequireShowLoginUI,
+	TEXT("True if login requires calling ShowLoginUI on the externalUI, depends on SupportExternalUI"),
+	ECVF_Default | ECVF_Cheat);
+
+bool FOnlineSubsystemNull::bForceShowLoginUIUserChange = false;
+FAutoConsoleVariableRef CVarForceShowLoginUIUserChange(
+	TEXT("OSSNull.ForceShowLoginUIUserChange"),
+	FOnlineSubsystemNull::bForceShowLoginUIUserChange,
+	TEXT("True if the user index should change during login UI to emulate a platform user change"),
+	ECVF_Default | ECVF_Cheat);
+
+bool FOnlineSubsystemNull::bRequireLoginCredentials = false;
+FAutoConsoleVariableRef CVarRequireLoginCredentials(
+	TEXT("OSSNull.RequireLoginCredentials"),
+	FOnlineSubsystemNull::bRequireLoginCredentials,
+	TEXT("True if login should require a user/pass to act like an external service, false to match most platforms and use the default"),
+	ECVF_Default | ECVF_Cheat);
+
+bool FOnlineSubsystemNull::bAddUserNumToNullId = false;
+FAutoConsoleVariableRef CVarAddUserNumToNullId(
+	TEXT("OSSNull.AddUserNumToNullId"),
+	FOnlineSubsystemNull::bAddUserNumToNullId,
+	TEXT("True if login name should include the local user number, which allows different stable Ids per user num"),
+	ECVF_Default | ECVF_Cheat);
+
+bool FOnlineSubsystemNull::bForceStableNullId = false;
+FAutoConsoleVariableRef CVarForceStableNullId(
+	TEXT("OSSNull.ForceStableNullId"),
+	FOnlineSubsystemNull::bForceStableNullId,
+	TEXT("True if it should use a system-stable null Id for login, same as -StableNullID on command line"),
+	ECVF_Default | ECVF_Cheat);
+
+bool FOnlineSubsystemNull::bForceOfflineMode = false;
+FAutoConsoleVariableRef CVarForceOfflineMode(
+	TEXT("OSSNull.ForceOfflineMode"),
+	FOnlineSubsystemNull::bForceOfflineMode,
+	TEXT("True if it should fail faked network queries and act like an offline system"),
+	ECVF_Default | ECVF_Cheat);
+
 bool FOnlineSubsystemNull::Init()
 {
 	const bool bNullInit = true;
-	
 	if (bNullInit)
 	{
+		static bool bLoadedConfig = false;
+		if (!bLoadedConfig)
+		{
+#if WITH_ENGINE
+			// Read configuration variables once for emulating other login systems, cvars can override in PIE
+			bLoadedConfig = true;
+			GConfig->GetBool(TEXT("OnlineSubsystemNull"), TEXT("bAutoLoginAtStartup"), bAutoLoginAtStartup, GEngineIni);
+			GConfig->GetBool(TEXT("OnlineSubsystemNull"), TEXT("bSupportExternalUI"), bSupportExternalUI, GEngineIni);
+			GConfig->GetBool(TEXT("OnlineSubsystemNull"), TEXT("bRequireShowLoginUI"), bRequireShowLoginUI, GEngineIni);
+			GConfig->GetBool(TEXT("OnlineSubsystemNull"), TEXT("bForceShowLoginUIUserChange"), bForceShowLoginUIUserChange, GEngineIni);
+			GConfig->GetBool(TEXT("OnlineSubsystemNull"), TEXT("bRequireLoginCredentials"), bRequireLoginCredentials, GEngineIni);
+			GConfig->GetBool(TEXT("OnlineSubsystemNull"), TEXT("bAddUserNumToNullId"), bAddUserNumToNullId, GEngineIni);
+			GConfig->GetBool(TEXT("OnlineSubsystemNull"), TEXT("bForceStableNullId"), bForceStableNullId, GEngineIni);
+			GConfig->GetBool(TEXT("OnlineSubsystemNull"), TEXT("bForceOfflineMode"), bForceOfflineMode, GEngineIni);
+
+			if (FParse::Param(FCommandLine::Get(), TEXT("StableNullID")))
+			{
+				bForceStableNullId = true;
+			}
+#endif
+		}
+
 		// Create the online async task thread
 		OnlineAsyncTaskThreadRunnable = new FOnlineAsyncTaskManagerNull(this);
 		check(OnlineAsyncTaskThreadRunnable);
@@ -209,6 +290,10 @@ bool FOnlineSubsystemNull::Init()
 		SessionInterface = MakeShareable(new FOnlineSessionNull(this));
 		LeaderboardsInterface = MakeShareable(new FOnlineLeaderboardsNull(this));
 		IdentityInterface = MakeShareable(new FOnlineIdentityNull(this));
+		if (bSupportExternalUI)
+		{
+			ExternalUIInterface = MakeShareable(new FOnlineExternalUINull(this));
+		}
 		AchievementsInterface = MakeShareable(new FOnlineAchievementsNull(this));
 #if WITH_ENGINE
 		VoiceInterface = MakeShareable(new FOnlineVoiceImpl(this));
diff --git a/Engine/Plugins/Online/OnlineSubsystemNull/Source/Public/OnlineSubsystemNull.h b/Engine/Plugins/Online/OnlineSubsystemNull/Source/Public/OnlineSubsystemNull.h
--- a/Engine/Plugins/Online/OnlineSubsystemNull/Source/Public/OnlineSubsystemNull.h
+++ b/Engine/Plugins/Online/OnlineSubsystemNull/Source/Public/OnlineSubsystemNull.h
@@ -99,6 +99,32 @@ PACKAGE_SCOPE:
 		OnlineAsyncTaskThread(nullptr)
 	{}
 
+	// Options for emulating different types of online platforms, these are settable via OSSNull cvars or in the [OnlineSubsystemNull] config section
+
+	/** True if it should login the first user at startup like single-user platforms, false to only login when requested */
+	static bool bAutoLoginAtStartup;
+
+	/** True if it should support an external UI interface */
+	static bool bSupportExternalUI;
+
+	/** True if login requires calling ShowLoginUI on the externalUI, depends on SupportExternalUI */
+	static bool bRequireShowLoginUI;
+
+	/** True if the user index should change during login UI to emulate a platform user change */
+	static bool bForceShowLoginUIUserChange;
+
+	/** True if login should require a user/pass to act like an external service, false to match most platforms and use the default */
+	static bool bRequireLoginCredentials;
+
+	/** True if login name should include the local user number, which allows different stable IDs per user num */
+	static bool bAddUserNumToNullId;
+
+	/** True if it should use a system-stable null Id for login, same as -StableNullID on command line */
+	static bool bForceStableNullId;
+
+	/** True if it should fail faked network queries and act like an offline system */
+	static bool bForceOfflineMode;
+
 private:
 
 	/** Interface to the session services */
@@ -116,6 +142,9 @@ private:
 	/** Interface to the identity registration/auth services */
 	FOnlineIdentityNullPtr IdentityInterface;
 
+	/** Interface to the identity registration/auth services */
+	FOnlineExternalUINullPtr ExternalUIInterface;
+
 	/** Interface for achievements */
 	FOnlineAchievementsNullPtr AchievementsInterface;
 
