OnlineSubsystem: Fix compilation issues when compiled without WITH_ENGINE.

diff --git a/Engine/Plugins/Online/OnlineSubsystemNull/Source/Private/OnlineSessionInterfaceNull.cpp b/Engine/Plugins/Online/OnlineSubsystemNull/Source/Private/OnlineSessionInterfaceNull.cpp
--- a/Engine/Plugins/Online/OnlineSubsystemNull/Source/Private/OnlineSessionInterfaceNull.cpp
+++ b/Engine/Plugins/Online/OnlineSubsystemNull/Source/Private/OnlineSessionInterfaceNull.cpp
@@ -1,14 +1,19 @@
 // Copyright Epic Games, Inc. All Rights Reserved.
 
 #include "OnlineSessionInterfaceNull.h"
+
+#include "UObject/CoreNet.h"
 #include "Online/OnlineBase.h"
+#include "Interfaces/OnlineIdentityInterface.h"
+
+#include "NboSerializerNull.h"
+#include "OnlineAsyncTaskManager.h"
+#include "OnlineSubsystem.h"
 #include "OnlineSubsystemNull.h"
 #if WITH_ENGINE
 #include "OnlineSubsystemUtils.h"
 #endif //WITH_ENGINE
-#include "OnlineAsyncTaskManager.h"
 #include "SocketSubsystem.h"
-#include "NboSerializerNull.h"
 
 
 
diff --git a/Engine/Plugins/Online/OnlineSubsystemNull/Source/Private/OnlineSubsystemNull.cpp b/Engine/Plugins/Online/OnlineSubsystemNull/Source/Private/OnlineSubsystemNull.cpp
--- a/Engine/Plugins/Online/OnlineSubsystemNull/Source/Private/OnlineSubsystemNull.cpp
+++ b/Engine/Plugins/Online/OnlineSubsystemNull/Source/Private/OnlineSubsystemNull.cpp
@@ -1,18 +1,23 @@
 // Copyright Epic Games, Inc. All Rights Reserved.
 
 #include "OnlineSubsystemNull.h"
+
 #include "HAL/RunnableThread.h"
-#include "OnlineAsyncTaskManagerNull.h"
 
-#include "OnlineSessionInterfaceNull.h"
-#include "OnlineLeaderboardInterfaceNull.h"
-#include "OnlineIdentityNull.h"
-#include "OnlineExternalUINull.h"
 #include "OnlineAchievementsInterfaceNull.h"
-#include "OnlineStoreV2InterfaceNull.h"
-#include "OnlinePurchaseInterfaceNull.h"
+#include "OnlineAsyncTaskManagerNull.h"
+#include "OnlineExternalUINull.h"
+#include "OnlineIdentityNull.h"
+#include "OnlineLeaderboardInterfaceNull.h"
 #include "OnlineMessageSanitizerNull.h"
+#include "OnlinePurchaseInterfaceNull.h"
+#include "OnlineSessionInterfaceNull.h"
+#include "OnlineStoreV2InterfaceNull.h"
+
+#if WITH_ENGINE
 #include "VoiceInterfaceImpl.h"
+#endif
+
 #include "Misc/CommandLine.h"
 
 #if WITH_ENGINE
