Included BuildUniqueId check in OSSNull to filter sessions found.

BuildUniqueId now uses FNetworkVersion::GetLocalNetworkVersion.

diff --git a/Engine/Plugins/Online/OnlineSubsystemNull/Source/Private/OnlineSessionInterfaceNull.cpp b/Engine/Plugins/Online/OnlineSubsystemNull/Source/Private/OnlineSessionInterfaceNull.cpp
--- a/Engine/Plugins/Online/OnlineSubsystemNull/Source/Private/OnlineSessionInterfaceNull.cpp
+++ b/Engine/Plugins/Online/OnlineSubsystemNull/Source/Private/OnlineSessionInterfaceNull.cpp
@@ -1191,22 +1191,32 @@ void FOnlineSessionNull::ReadSettingsFromPacket(FNboSerializeFromBufferNull& Pac
 
 void FOnlineSessionNull::OnValidResponsePacketReceived(uint8* PacketData, int32 PacketLength)
 {
-	// Create an object that we'll copy the data to
-	FOnlineSessionSettings NewServer;
 	if (CurrentSessionSearch.IsValid())
 	{
-		// Add space in the search results array
-		FOnlineSessionSearchResult* NewResult = new (CurrentSessionSearch->SearchResults) FOnlineSessionSearchResult();
-		// this is not a correct ping, but better than nothing
-		NewResult->PingInMs = static_cast<int32>((FPlatformTime::Seconds() - SessionSearchStartInSeconds) * 1000);
+		FOnlineSessionSearchResult SearchResult;
 
-		FOnlineSession* NewSession = &NewResult->Session;
+		// this is not a correct ping, but better than nothing
+		SearchResult.PingInMs = static_cast<int32>((FPlatformTime::Seconds() - SessionSearchStartInSeconds) * 1000);
 
 		// Prepare to read data from the packet
 		FNboSerializeFromBufferNull Packet(PacketData, PacketLength);
-		
-		ReadSessionFromPacket(Packet, NewSession);
 
+		ReadSessionFromPacket(Packet, &SearchResult.Session);
+
+		const int32 BuildUniqueId = GetBuildUniqueId();
+
+		if (SearchResult.Session.SessionSettings.BuildUniqueId == BuildUniqueId)
+		{
+			// Add the found session to the search results array
+			CurrentSessionSearch->SearchResults.Emplace(MoveTemp(SearchResult));
+		}
+		else
+		{
+			UE_LOG_ONLINE_SESSION(Verbose, TEXT("Rejecting search result [%s]: mismatched build id. Local: %d - Session: %d"),
+				*SearchResult.Session.GetSessionIdStr(),
+				BuildUniqueId,
+				SearchResult.Session.SessionSettings.BuildUniqueId);
+		}
 		// NOTE: we don't notify until the timeout happens
 	}
 	else
