Add bOnlineRequiresSecondLogin cvar to OSS Null for testing the use case of a system being offline when it is first launched

diff --git a/Engine/Plugins/Online/OnlineSubsystemNull/Source/Private/OnlineIdentityNull.cpp b/Engine/Plugins/Online/OnlineSubsystemNull/Source/Private/OnlineIdentityNull.cpp
--- a/Engine/Plugins/Online/OnlineSubsystemNull/Source/Private/OnlineIdentityNull.cpp
+++ b/Engine/Plugins/Online/OnlineSubsystemNull/Source/Private/OnlineIdentityNull.cpp
@@ -9,6 +9,8 @@
 #include "SocketSubsystem.h"
 #include "OnlineError.h"
 
+#define USER_ATTR_NULL_LOGINCOUNT TEXT("null:logincount")
+
 bool FUserOnlineAccountNull::GetAuthAttribute(const FString& AttrName, FString& OutAttrValue) const
 {
 	const FString* FoundAttr = AdditionalAuthData.Find(AttrName);
@@ -125,6 +127,13 @@ bool FOnlineIdentityNull::LoginInternal(int32 LocalUserNum, const FOnlineAccount
 		{
 			UserAccountPtr = UserAccounts.FindChecked(UserId->ToSharedRef());
 		}
+
+		// Increment the login count
+		FString LoginCountString = TEXT("0");
+		UserAccountPtr->GetUserAttribute(USER_ATTR_NULL_LOGINCOUNT, LoginCountString);
+		int32 LoginCount = FCString::Atoi(*LoginCountString) + 1;
+		LoginCountString = FString::Printf(TEXT("%d"), LoginCount);
+		UserAccountPtr->SetUserAttribute(USER_ATTR_NULL_LOGINCOUNT, LoginCountString);
 	}
 
 	if (!ErrorStr.IsEmpty())
@@ -273,6 +282,16 @@ ELoginStatus::Type FOnlineIdentityNull::GetLoginStatus(const FUniqueNetId& UserI
 		{
 			return ELoginStatus::UsingLocalProfile;
 		}
+		else if (FOnlineSubsystemNull::bOnlineRequiresSecondLogin)
+		{
+			FString LoginCountString = TEXT("0");
+			UserAccount->GetUserAttribute(USER_ATTR_NULL_LOGINCOUNT, LoginCountString);
+			int32 LoginCount = FCString::Atoi(*LoginCountString);
+			if (LoginCount < 2)
+			{
+				return ELoginStatus::UsingLocalProfile;
+			}
+		}
 		return ELoginStatus::LoggedIn;
 	}
 	return ELoginStatus::NotLoggedIn;
diff --git a/Engine/Plugins/Online/OnlineSubsystemNull/Source/Private/OnlineSubsystemNull.cpp b/Engine/Plugins/Online/OnlineSubsystemNull/Source/Private/OnlineSubsystemNull.cpp
--- a/Engine/Plugins/Online/OnlineSubsystemNull/Source/Private/OnlineSubsystemNull.cpp
+++ b/Engine/Plugins/Online/OnlineSubsystemNull/Source/Private/OnlineSubsystemNull.cpp
@@ -257,6 +257,13 @@ FAutoConsoleVariableRef CVarForceOfflineMode(
 	TEXT("True if it should fail faked network queries and act like an offline system"),
 	ECVF_Default | ECVF_Cheat);
 
+bool FOnlineSubsystemNull::bOnlineRequiresSecondLogin = false;
+FAutoConsoleVariableRef CVarOnlineRequiresSecondLogin(
+	TEXT("OSSNull.OnlineRequiresSecondLogin"),
+	FOnlineSubsystemNull::bOnlineRequiresSecondLogin,
+	TEXT("True if the first login only counts as local login, a second is required for online access"),
+	ECVF_Default | ECVF_Cheat);
+
 bool FOnlineSubsystemNull::Init()
 {
 	const bool bNullInit = true;
@@ -276,6 +283,7 @@ bool FOnlineSubsystemNull::Init()
 			GConfig->GetBool(TEXT("OnlineSubsystemNull"), TEXT("bAddUserNumToNullId"), bAddUserNumToNullId, GEngineIni);
 			GConfig->GetBool(TEXT("OnlineSubsystemNull"), TEXT("bForceStableNullId"), bForceStableNullId, GEngineIni);
 			GConfig->GetBool(TEXT("OnlineSubsystemNull"), TEXT("bForceOfflineMode"), bForceOfflineMode, GEngineIni);
+			GConfig->GetBool(TEXT("OnlineSubsystemNull"), TEXT("bOnlineRequiresSecondLogin"), bOnlineRequiresSecondLogin, GEngineIni);
 
 			if (FParse::Param(FCommandLine::Get(), TEXT("StableNullID")))
 			{
diff --git a/Engine/Plugins/Online/OnlineSubsystemNull/Source/Public/OnlineSubsystemNull.h b/Engine/Plugins/Online/OnlineSubsystemNull/Source/Public/OnlineSubsystemNull.h
--- a/Engine/Plugins/Online/OnlineSubsystemNull/Source/Public/OnlineSubsystemNull.h
+++ b/Engine/Plugins/Online/OnlineSubsystemNull/Source/Public/OnlineSubsystemNull.h
@@ -125,6 +125,9 @@ PACKAGE_SCOPE:
 	/** True if it should fail faked network queries and act like an offline system */
 	static bool bForceOfflineMode;
 
+	/** True if the first login only counts as local login, a second is required for online access */
+	static bool bOnlineRequiresSecondLogin;
+
 private:
 
 	/** Interface to the session services */
