PR #11577: Update OnlineSessionInterfaceNull.cpp

diff --git a/Engine/Plugins/Online/OnlineSubsystemNull/Source/Private/OnlineSessionInterfaceNull.cpp b/Engine/Plugins/Online/OnlineSubsystemNull/Source/Private/OnlineSessionInterfaceNull.cpp
--- a/Engine/Plugins/Online/OnlineSubsystemNull/Source/Private/OnlineSessionInterfaceNull.cpp
+++ b/Engine/Plugins/Online/OnlineSubsystemNull/Source/Private/OnlineSessionInterfaceNull.cpp
@@ -247,7 +247,7 @@ bool FOnlineSessionNull::CreateSession(int32 HostingPlayerNum, FName SessionName
 
 	if (Result != ONLINE_IO_PENDING)
 	{
-		TriggerOnCreateSessionCompleteDelegates(SessionName, (Result == ONLINE_SUCCESS) ? true : false);
+		TriggerOnCreateSessionCompleteDelegates(SessionName, (Result == ONLINE_SUCCESS));
 	}
 	
 	return Result == ONLINE_IO_PENDING || Result == ONLINE_SUCCESS;
