Added missing value to session information in Null OSS Sessions

diff --git a/Engine/Plugins/Online/OnlineSubsystemNull/Source/Private/OnlineSessionInterfaceNull.cpp b/Engine/Plugins/Online/OnlineSubsystemNull/Source/Private/OnlineSessionInterfaceNull.cpp
--- a/Engine/Plugins/Online/OnlineSubsystemNull/Source/Private/OnlineSessionInterfaceNull.cpp
+++ b/Engine/Plugins/Online/OnlineSubsystemNull/Source/Private/OnlineSessionInterfaceNull.cpp
@@ -196,6 +196,7 @@ bool FOnlineSessionNull::CreateSession(int32 HostingPlayerNum, FName SessionName
 		// Create a new session and deep copy the game settings
 		Session = AddNamedSession(SessionName, NewSessionSettings);
 		check(Session);
+		Session->bHosting = true;
 		Session->SessionState = EOnlineSessionState::Creating;
 		Session->NumOpenPrivateConnections = NewSessionSettings.NumPrivateConnections;
 		Session->NumOpenPublicConnections = NewSessionSettings.NumPublicConnections;	// always start with full public connections, local player will register later
