AutoRTFM:  Allowing OnlineSessionInterfaceNull functions to be called in transactions by replacing the FCriticalSections with FTransactionallySafeCriticalSection.  These functions are typically called deep in the networking stack (e.g. in response to NMT_Join).

diff --git a/Engine/Plugins/Online/OnlineSubsystemNull/Source/Private/OnlineSessionInterfaceNull.cpp b/Engine/Plugins/Online/OnlineSubsystemNull/Source/Private/OnlineSessionInterfaceNull.cpp
--- a/Engine/Plugins/Online/OnlineSubsystemNull/Source/Private/OnlineSessionInterfaceNull.cpp
+++ b/Engine/Plugins/Online/OnlineSubsystemNull/Source/Private/OnlineSessionInterfaceNull.cpp
@@ -262,7 +262,7 @@ bool FOnlineSessionNull::CreateSession(const FUniqueNetId& HostingPlayerId, FNam
 
 bool FOnlineSessionNull::NeedsToAdvertise()
 {
-	FScopeLock ScopeLock(&SessionLock);
+	UE::TScopeLock ScopeLock(SessionLock);
 
 	bool bResult = false;
 	for (int32 SessionIdx=0; SessionIdx < Sessions.Num(); SessionIdx++)
@@ -1065,7 +1065,7 @@ void FOnlineSessionNull::AppendSessionSettingsToPacket(FNboSerializeToBufferNull
 void FOnlineSessionNull::OnValidQueryPacketReceived(uint8* PacketData, int32 PacketLength, uint64 ClientNonce)
 {
 	// Iterate through all registered sessions and respond for each one that can be joinable
-	FScopeLock ScopeLock(&SessionLock);
+	UE::TScopeLock ScopeLock(SessionLock);
 	for (int32 SessionIndex = 0; SessionIndex < Sessions.Num(); SessionIndex++)
 	{
 		FNamedOnlineSession* Session = &Sessions[SessionIndex];
@@ -1255,13 +1255,13 @@ void FOnlineSessionNull::OnLANSearchTimeout()
 
 int32 FOnlineSessionNull::GetNumSessions()
 {
-	FScopeLock ScopeLock(&SessionLock);
+	UE::TScopeLock ScopeLock(SessionLock);
 	return Sessions.Num();
 }
 
 void FOnlineSessionNull::DumpSessionState()
 {
-	FScopeLock ScopeLock(&SessionLock);
+	UE::TScopeLock ScopeLock(SessionLock);
 
 	for (int32 SessionIdx=0; SessionIdx < Sessions.Num(); SessionIdx++)
 	{
diff --git a/Engine/Plugins/Online/OnlineSubsystemNull/Source/Private/OnlineSessionInterfaceNull.h b/Engine/Plugins/Online/OnlineSubsystemNull/Source/Private/OnlineSessionInterfaceNull.h
--- a/Engine/Plugins/Online/OnlineSubsystemNull/Source/Private/OnlineSessionInterfaceNull.h
+++ b/Engine/Plugins/Online/OnlineSubsystemNull/Source/Private/OnlineSessionInterfaceNull.h
@@ -162,7 +162,7 @@ private:
 PACKAGE_SCOPE:
 
 	/** Critical sections for thread safe operation of session lists */
-	mutable FCriticalSection SessionLock;
+	mutable FTransactionallySafeCriticalSection SessionLock;
 
 	/** Current session settings */
 	TArray<FNamedOnlineSession> Sessions;
@@ -187,13 +187,13 @@ PACKAGE_SCOPE:
 	// IOnlineSession
 	class FNamedOnlineSession* AddNamedSession(FName SessionName, const FOnlineSessionSettings& SessionSettings) override
 	{
-		FScopeLock ScopeLock(&SessionLock);
+		UE::TScopeLock ScopeLock(SessionLock);
 		return &Sessions.Emplace_GetRef(SessionName, SessionSettings);
 	}
 
 	class FNamedOnlineSession* AddNamedSession(FName SessionName, const FOnlineSession& Session) override
 	{
-		FScopeLock ScopeLock(&SessionLock);
+		UE::TScopeLock ScopeLock(SessionLock);
 		return &Sessions.Emplace_GetRef(SessionName, Session);
 	}
 
@@ -231,7 +231,7 @@ public:
 
 	FNamedOnlineSession* GetNamedSession(FName SessionName) override
 	{
-		FScopeLock ScopeLock(&SessionLock);
+		UE::TScopeLock ScopeLock(SessionLock);
 		for (int32 SearchIndex = 0; SearchIndex < Sessions.Num(); SearchIndex++)
 		{
 			if (Sessions[SearchIndex].SessionName == SessionName)
@@ -244,7 +244,7 @@ public:
 
 	virtual void RemoveNamedSession(FName SessionName) override
 	{
-		FScopeLock ScopeLock(&SessionLock);
+		UE::TScopeLock ScopeLock(SessionLock);
 		for (int32 SearchIndex = 0; SearchIndex < Sessions.Num(); SearchIndex++)
 		{
 			if (Sessions[SearchIndex].SessionName == SessionName)
@@ -257,7 +257,7 @@ public:
 
 	virtual EOnlineSessionState::Type GetSessionState(FName SessionName) const override
 	{
-		FScopeLock ScopeLock(&SessionLock);
+		UE::TScopeLock ScopeLock(SessionLock);
 		for (int32 SearchIndex = 0; SearchIndex < Sessions.Num(); SearchIndex++)
 		{
 			if (Sessions[SearchIndex].SessionName == SessionName)
@@ -271,7 +271,7 @@ public:
 
 	virtual bool HasPresenceSession() override
 	{
-		FScopeLock ScopeLock(&SessionLock);
+		UE::TScopeLock ScopeLock(SessionLock);
 		for (int32 SearchIndex = 0; SearchIndex < Sessions.Num(); SearchIndex++)
 		{
 			if (Sessions[SearchIndex].SessionSettings.bUsesPresence)
