- Added OnlineTitleFileNull which can be used to load hotfixes in -nomcp mode (pulled code out of TitleFileMcp into the Null one)
- Changed OnlineHotfixManager to use the OnlineTitleFileNull when -TitleFileLocalPath or -HotfixLocalPath are specified

#Synced-CL 48456107

diff --git a/Engine/Plugins/Online/OnlineFramework/Source/Hotfix/Private/OnlineHotfixManager.cpp b/Engine/Plugins/Online/OnlineFramework/Source/Hotfix/Private/OnlineHotfixManager.cpp
--- a/Engine/Plugins/Online/OnlineFramework/Source/Hotfix/Private/OnlineHotfixManager.cpp
+++ b/Engine/Plugins/Online/OnlineFramework/Source/Hotfix/Private/OnlineHotfixManager.cpp
@@ -45,6 +45,10 @@ DEFINE_LOG_CATEGORY(LogHotfixManager);
 /** The prefix used when there no custom config so we don't try to match empty strings */
 #define HOTFIX_NO_CUSTOM_CONFIG TEXT("##NOCONFIG##")
 
+// do we support loading local HF files?
+#define UE_TITLE_FILE_DEV_LOCAL_FILES_ENABLED !UE_BUILD_SHIPPING
+
+
 FName NAME_HotfixManager(TEXT("HotfixManager"));
 
 static TAutoConsoleVariable<int32> CVarUseNewDynamicLayersForHotfix(
@@ -287,6 +291,30 @@ void UOnlineHotfixManager::PostInitProperties()
 	// Build the default prefix too
 	DefaultPrefix = DebugPrefix + TEXT("Default");
 
+#if UE_TITLE_FILE_DEV_LOCAL_FILES_ENABLED
+	// if we are going to load from local hotfix files, use the OnlineTitleFileNull class, which only loads from local files
+	FString TitleFileLocalPath, HotfixLocalPath;
+	if (FParse::Value(FCommandLine::Get(), TEXT("TitleFileLocalPath="), TitleFileLocalPath) || FParse::Value(FCommandLine::Get(), TEXT("HotfixLocalPath="), HotfixLocalPath))
+	{
+		// if HotfixLocalPath is specified, check that it either matches an existing TitleFileLocalPath, or if it doesn't exist, then alias it to 
+		// TitleFileLocalPath so that the OnlineTitleFileNull code can use the proper option
+		if (!HotfixLocalPath.IsEmpty())
+		{
+			if (TitleFileLocalPath.IsEmpty())
+			{
+			    FString LocalCommandLine = FString::Printf(TEXT(" TitleFileLocalPath=\"%s\""), *HotfixLocalPath);
+			    FCommandLine::Append(*LocalCommandLine);
+			}
+            else
+            {
+			    ensure(HotfixLocalPath == TitleFileLocalPath);
+			}
+		}	
+    
+		OSSName = "Null";
+	}
+#endif
+
 	Super::PostInitProperties();
 }
 
diff --git a/Engine/Plugins/Online/OnlineSubsystemNull/Source/Private/OnlineSubsystemNull.cpp b/Engine/Plugins/Online/OnlineSubsystemNull/Source/Private/OnlineSubsystemNull.cpp
--- a/Engine/Plugins/Online/OnlineSubsystemNull/Source/Private/OnlineSubsystemNull.cpp
+++ b/Engine/Plugins/Online/OnlineSubsystemNull/Source/Private/OnlineSubsystemNull.cpp
@@ -4,6 +4,7 @@
 
 #include "HAL/RunnableThread.h"
 #include "Misc/Fork.h"
+#include "Misc/FileHelper.h"
 
 #include "OnlineAchievementsInterfaceNull.h"
 #include "OnlineAsyncTaskManagerNull.h"
@@ -14,6 +15,7 @@
 #include "OnlinePurchaseInterfaceNull.h"
 #include "OnlineSessionInterfaceNull.h"
 #include "OnlineStoreV2InterfaceNull.h"
+#include "OnlineTitleFileNull.h"
 
 #if WITH_ENGINE
 #include "VoiceInterfaceImpl.h"
@@ -103,7 +105,7 @@ IOnlineIdentityPtr FOnlineSubsystemNull::GetIdentityInterface() const
 
 IOnlineTitleFilePtr FOnlineSubsystemNull::GetTitleFileInterface() const
 {
-	return nullptr;
+	return TitleFileInterface;
 }
 
 IOnlineStoreV2Ptr FOnlineSubsystemNull::GetStoreV2Interface() const
@@ -190,6 +192,11 @@ bool FOnlineSubsystemNull::Tick(float DeltaTime)
 		SessionInterface->Tick(DeltaTime);
 	}
 
+	if (TitleFileInterface.IsValid())
+	{
+		TitleFileInterface->Tick(DeltaTime);
+	}
+
 #if WITH_ENGINE
 	if (VoiceInterface.IsValid() && bVoiceInterfaceInitialized)
 	{
@@ -314,6 +321,7 @@ bool FOnlineSubsystemNull::Init()
 		StoreV2Interface = MakeShareable(new FOnlineStoreV2Null(*this));
 		PurchaseInterface = MakeShareable(new FOnlinePurchaseNull(*this));
 		MessageSanitizerInterface = MakeShareable(new FMessageSanitizerNull());
+		TitleFileInterface = MakeShareable(new FOnlineTitleFileNull());
 	}
 	else
 	{
@@ -365,6 +373,7 @@ bool FOnlineSubsystemNull::Shutdown()
 	DESTRUCT_INTERFACE(LeaderboardsInterface);
 	DESTRUCT_INTERFACE(SessionInterface);
 	DESTRUCT_INTERFACE(MessageSanitizerInterface);
+	DESTRUCT_INTERFACE(TitleFileInterface);
 
 #undef DESTRUCT_INTERFACE
 	
diff --git a/Engine/Plugins/Online/OnlineSubsystemNull/Source/Private/OnlineTitleFileNull.cpp b/Engine/Plugins/Online/OnlineSubsystemNull/Source/Private/OnlineTitleFileNull.cpp
new file mode 100644
--- /dev/null
+++ b/Engine/Plugins/Online/OnlineSubsystemNull/Source/Private/OnlineTitleFileNull.cpp
@@ -0,0 +1,287 @@
+// Copyright Epic Games, Inc. All Rights Reserved.
+
+#include "OnlineTitleFileNull.h"
+#include "Misc/CommandLine.h"
+#include "Misc/ConfigCacheIni.h"
+#include "HAL/FileManager.h"
+#include "OnlineSubsystemNull.h"
+
+
+void FTitleFileNullAsyncLoadLocalDevFile::DoWork()
+{
+	// load file from disk
+	bool bLoadedFile = false;
+
+	FArchive* Reader = IFileManager::Get().CreateFileReader(*LocalDevFilePath, FILEREAD_Silent);
+	if (Reader)
+	{
+		int64 SizeToRead = Reader->TotalSize();
+		FileData.Reset(SizeToRead);
+		FileData.AddUninitialized(SizeToRead);
+
+		uint8* FileDataPtr = FileData.GetData();
+
+		static const int64 ChunkSize = 100 * 1024;
+
+		int64 TotalBytesRead = 0;
+		while (SizeToRead > 0)
+		{
+			int64 Val = FMath::Min(SizeToRead, ChunkSize);
+			Reader->Serialize(FileDataPtr + TotalBytesRead, Val);
+			BytesRead.Add(Val);
+			TotalBytesRead += Val;
+			SizeToRead -= Val;
+		}
+
+		ensure(SizeToRead == 0 && Reader->TotalSize() == TotalBytesRead);
+		bLoadedFile = Reader->Close();
+		delete Reader;
+	}
+	if (bLoadedFile)
+	{
+		UE_LOG_ONLINE(Verbose, TEXT("async local ini read completed for file '%s'"), *LocalDevFilePath);
+	}
+	else
+	{
+		UE_LOG_ONLINE(Verbose, TEXT("async local ini read failed for file '%s'"), *LocalDevFilePath);
+	}
+}
+
+
+
+
+FOnlineTitleFileNull::FOnlineTitleFileNull()
+{
+	// check for the option that is needed for this to work (if it's not, EnumerateFiles will do nothing
+	FParse::Value(FCommandLine::Get(), TEXT("TitleFileLocalPath="), LocalPath);
+}
+
+bool FOnlineTitleFileNull::GetFileContents(const FString& FileName, TArray<uint8>& FileContents)
+{
+	const TArray<FCloudFile>* FilesPtr = &Files;
+	if (FilesPtr != NULL)
+	{
+		for (TArray<FCloudFile>::TConstIterator It(*FilesPtr); It; ++It)
+		{
+			if (It->FileName == FileName)
+			{
+				FileContents = It->Data;
+				return true;
+			}
+		}
+	}
+	return false;
+}
+
+bool FOnlineTitleFileNull::ClearFiles()
+{
+	for (int Idx = 0; Idx < Files.Num(); Idx++)
+	{
+		if (Files[Idx].AsyncState == EOnlineAsyncTaskState::InProgress)
+		{
+			UE_LOG_ONLINE(Warning, TEXT("can't clear all files - '%s' has a pending async task!"), *Files[Idx].FileName);
+			return false;
+		}
+	}
+	// remove all cached file entries
+	Files.Empty();
+	return true;
+}
+
+bool FOnlineTitleFileNull::ClearFile(const FString& FileName)
+{
+	for (int Idx = 0; Idx < Files.Num(); Idx++)
+	{
+		if (Files[Idx].FileName == FileName)
+		{
+			if (Files[Idx].AsyncState == EOnlineAsyncTaskState::InProgress)
+			{
+				UE_LOG_ONLINE(Warning, TEXT("can't clear file '%s' - it has a pending async task!"), *Files[Idx].FileName);
+				return false;
+			}
+			else
+			{
+				Files.RemoveAt(Idx);
+				return true;
+			}
+		}
+	}
+	return false;
+}
+
+void FOnlineTitleFileNull::DeleteCachedFiles(bool bSkipEnumerated)
+{
+	// no cache to delete!
+}
+
+bool FOnlineTitleFileNull::EnumerateFiles(const FPagedQuery& Page)
+{
+	if (LocalPath.IsEmpty())
+	{
+		TriggerOnEnumerateFilesCompleteDelegates(false, TEXT("-TitleFileLocalPath= was not specified on the commandline."));
+		return false;
+	}
+
+	if (!IFileManager::Get().DirectoryExists(*LocalPath))
+	{
+		TriggerOnEnumerateFilesCompleteDelegates(false, TEXT("-TitleFileLocalPath= was specified, but does not exist."));
+		return false;
+	}
+
+	static int32 sFakeHash = 0;
+	//		bIsLoadingLocalDevFiles = true;
+	TArray<FString> FoundFiles;
+	IFileManager::Get().FindFiles(FoundFiles, *LocalPath, TEXT("*.*"));
+	UE_LOG_ONLINE(Verbose, TEXT("EnumerateFiles using local file path founds %d files at '%s'"), FoundFiles.Num(), *LocalPath);
+
+	const FName NAME_SHA1(TEXT("SHA1"));
+
+	FileHeaders.Empty();
+	for (const FString& CurrFile : FoundFiles)
+	{
+		FString FullPath = FPaths::Combine(LocalPath, CurrFile);
+		uint8 Hash[20];
+		FSHA1::GetFileSHAHash(*FullPath, Hash, false);
+
+		FCloudFileHeader FileHeader;
+		FileHeader.Hash = FString::Printf(TEXT("%d"), ++sFakeHash);
+		FileHeader.HashType = NAME_SHA1;
+		FileHeader.DLName = FullPath;
+		FileHeader.FileName = CurrFile;
+		FileHeader.URL = "File://" + FullPath;
+		FileHeader.ChunkID = 0;
+		int64 FileSize = IFileManager::Get().FileSize(*FullPath);
+		check(FileSize < INT_MAX);
+		FileHeader.FileSize = (int32)(FileSize);
+		// FileHeader.ExternalStorageIds; Leave empty
+
+		FileHeaders.Add(FileHeader);
+	}
+	// Note: In the canonical cloud file case, this would happen at a latter time and not be on the same stack
+	// that EnumerateFiles is called on.
+	TriggerOnEnumerateFilesCompleteDelegates(true, TEXT(""));
+	return true;
+}
+
+void FOnlineTitleFileNull::GetFileList(TArray<FCloudFileHeader>& OutFiles)
+{
+	TArray<FCloudFileHeader>* FilesPtr = &FileHeaders;
+	if (FilesPtr != NULL)
+	{
+		OutFiles = *FilesPtr;
+	}
+	else
+	{
+		OutFiles.Empty();
+	}
+}
+
+bool FOnlineTitleFileNull::ReadFile(const FString& FileName)
+{
+	bool bStarted = false;
+
+	FCloudFile* CloudFile = GetCloudFile(FileName, true);
+	if (CloudFile->AsyncState == EOnlineAsyncTaskState::InProgress)
+	{
+		UE_LOG_ONLINE(Verbose, TEXT("file read (local) is already in progress for file '%s'"), *FileName);
+		return true;
+	}
+
+	FCloudFileHeader* CloudFileHeader = GetCloudFileHeader(FileName);
+	if (CloudFileHeader != nullptr)
+	{
+		FAsyncTask<FTitleFileNullAsyncLoadLocalDevFile>* NewRead =
+			new FAsyncTask<FTitleFileNullAsyncLoadLocalDevFile>(CloudFileHeader->FileName, CloudFileHeader->DLName);
+		AsyncLocalDevReads.Add(NewRead);
+		CloudFile->AsyncState = EOnlineAsyncTaskState::InProgress;
+		NewRead->StartBackgroundTask();
+		bStarted = true;
+	}
+	return bStarted;
+}
+
+void FOnlineTitleFileNull::Tick(float DeltaTime)
+{
+	QUICK_SCOPE_CYCLE_COUNTER(STAT_FOnlineTitleFileMcpV2_Tick);
+	LLM_SCOPE_BYTAG(OnlineSubsystem);
+
+	TArray<int32> ItemsToRemove;
+	ItemsToRemove.Reserve(AsyncLocalDevReads.Num());
+
+	// Local files, developer flow
+	for (int32 TaskIdx = 0; TaskIdx < AsyncLocalDevReads.Num(); TaskIdx++)
+	{
+		FAsyncTask<FTitleFileNullAsyncLoadLocalDevFile>* ReadTask = AsyncLocalDevReads[TaskIdx];
+		if (ReadTask->IsDone())
+		{
+			ItemsToRemove.Add(TaskIdx);
+			FinishReadFileLocalDevFile(ReadTask);
+		}
+	}
+
+	// Clean up any tasks that were completed
+	for (int32 ItemIdx = ItemsToRemove.Num() - 1; ItemIdx >= 0; ItemIdx--)
+	{
+		const int32 TaskIdx = ItemsToRemove[ItemIdx];
+		FAsyncTask<FTitleFileNullAsyncLoadLocalDevFile>* TaskToDelete = AsyncLocalDevReads[TaskIdx];
+		const FTitleFileNullAsyncLoadLocalDevFile& TaskInfo = TaskToDelete->GetTask();
+		UE_LOG_ONLINE(VeryVerbose, TEXT("removed local dev file read task for: %s (bytes read: %d)"), *TaskInfo.LocalDevFileName, TaskInfo.BytesRead.GetValue());
+		delete TaskToDelete;
+		AsyncLocalDevReads.RemoveAtSwap(TaskIdx);
+	}
+	ItemsToRemove.Empty();
+}
+
+FCloudFileHeader* FOnlineTitleFileNull::GetCloudFileHeader(const FString& FileName)
+{
+	FCloudFileHeader* CloudFileHeader = NULL;
+
+	for (int Idx = 0; Idx < FileHeaders.Num(); Idx++)
+	{
+		if (FileHeaders[Idx].DLName == FileName)
+		{
+			CloudFileHeader = &FileHeaders[Idx];
+			break;
+		}
+	}
+
+	return CloudFileHeader;
+}
+
+FCloudFile* FOnlineTitleFileNull::GetCloudFile(const FString& FileName, bool bCreateIfMissing)
+{
+	FCloudFile* CloudFile = NULL;
+	for (int Idx = 0; Idx < Files.Num(); Idx++)
+	{
+		if (Files[Idx].FileName == FileName)
+		{
+			CloudFile = &Files[Idx];
+			break;
+		}
+	}
+	if (CloudFile == NULL && bCreateIfMissing)
+	{
+		CloudFile = new(Files)FCloudFile(FileName);
+	}
+	return CloudFile;
+}
+
+/** Completes the async operation for the local dev file read. Similar to above but without any cloud connection. */
+void FOnlineTitleFileNull::FinishReadFileLocalDevFile(FAsyncTask<FTitleFileNullAsyncLoadLocalDevFile>* Task)
+{
+	const FTitleFileNullAsyncLoadLocalDevFile& TaskInfo = Task->GetTask();
+	FCloudFileHeader* CloudFileHeader = GetCloudFileHeader(TaskInfo.LocalDevFilePath);
+	FCloudFile* CloudFile = GetCloudFile(TaskInfo.LocalDevFilePath, true);
+	if (CloudFileHeader != nullptr && CloudFile != nullptr)
+	{
+		CloudFile->Data = TaskInfo.FileData;
+		CloudFile->AsyncState = EOnlineAsyncTaskState::Done;
+		CloudFileHeader->Hash = TaskInfo.LocalHashStr;
+		TriggerOnReadFileCompleteDelegates(true, TaskInfo.LocalDevFilePath);
+	}
+	else
+	{
+		UE_LOG_ONLINE(Warning, TEXT("failed to read '%s' from given local Dev path!"), *TaskInfo.LocalDevFileName);
+	}
+}
+	
\ No newline at end of file
diff --git a/Engine/Plugins/Online/OnlineSubsystemNull/Source/Private/OnlineTitleFileNull.h b/Engine/Plugins/Online/OnlineSubsystemNull/Source/Private/OnlineTitleFileNull.h
new file mode 100644
--- /dev/null
+++ b/Engine/Plugins/Online/OnlineSubsystemNull/Source/Private/OnlineTitleFileNull.h
@@ -0,0 +1,75 @@
+// Copyright Epic Games, Inc. All Rights Reserved.
+
+#pragma once
+
+#include "Async/AsyncWork.h"
+#include "Interfaces/OnlineTitleFileInterface.h"
+
+class FOnlineSubsystemNull;
+
+class FTitleFileNullAsyncLoadLocalDevFile :
+	public FNonAbandonableTask
+{
+public:
+	/** File data loaded for the async read */
+	TArray<uint8> FileData;
+	/** Amount of data read from the file to be owned/referenced by the game thread */
+	FThreadSafeCounter64 BytesRead;
+	/** The original name of the file being read */
+	const FString LocalDevFileName;
+	/** The full path file being read */
+	const FString LocalDevFilePath;
+	/** Hash Calculated after file is loaded */
+	FString LocalHashStr;
+
+	FTitleFileNullAsyncLoadLocalDevFile(const FString& FileName, const FString& FullPath) :
+		LocalDevFileName(FileName),
+		LocalDevFilePath(FullPath)
+	{
+
+	}
+
+	void DoWork();
+
+	FORCEINLINE TStatId GetStatId() const
+	{
+		RETURN_QUICK_DECLARE_CYCLE_STAT(FTitleFileNullAsyncLoadLocalDevFile, STATGROUP_ThreadPoolAsyncTasks);
+	}
+
+};
+
+class FOnlineTitleFileNull : public IOnlineTitleFile, public TSharedFromThis<FOnlineTitleFileNull, ESPMode::ThreadSafe>
+{
+public:
+
+	FOnlineTitleFileNull();
+	virtual ~FOnlineTitleFileNull() = default;
+
+
+	virtual bool GetFileContents(const FString& FileName, TArray<uint8>& FileContents) override;
+	virtual bool ClearFiles() override;
+	virtual bool ClearFile(const FString& FileName) override;
+	virtual void DeleteCachedFiles(bool bSkipEnumerated) override;
+	virtual bool EnumerateFiles(const FPagedQuery& Page) override;
+	virtual void GetFileList(TArray<FCloudFileHeader>& OutFiles) override;
+	virtual bool ReadFile(const FString& FileName) override;
+
+	void Tick(float DeltaTime);
+
+protected:
+	FString LocalPath;
+	TArray<FCloudFileHeader> FileHeaders;
+	TArray<FCloudFile> Files;
+
+	/** Holds the outstanding task for hitch free loading of local Dev files. */
+	TArray<FAsyncTask<FTitleFileNullAsyncLoadLocalDevFile>* > AsyncLocalDevReads;
+
+
+	// these are all copied from OnlineTitleMcp, and maybe could be shared somehow
+	FCloudFileHeader* GetCloudFileHeader(const FString& FileName);
+	FCloudFile* GetCloudFile(const FString& FileName, bool bCreateIfMissing);
+	/** Completes the async operation for the local dev file read. Similar to above but without any cloud connection. */
+	void FinishReadFileLocalDevFile(FAsyncTask<FTitleFileNullAsyncLoadLocalDevFile>* Task);
+};
+
+using FOnlineTitleFileNullPtr = TSharedPtr<FOnlineTitleFileNull, ESPMode::ThreadSafe>;
diff --git a/Engine/Plugins/Online/OnlineSubsystemNull/Source/Public/OnlineSubsystemNull.h b/Engine/Plugins/Online/OnlineSubsystemNull/Source/Public/OnlineSubsystemNull.h
--- a/Engine/Plugins/Online/OnlineSubsystemNull/Source/Public/OnlineSubsystemNull.h
+++ b/Engine/Plugins/Online/OnlineSubsystemNull/Source/Public/OnlineSubsystemNull.h
@@ -27,6 +27,7 @@ typedef TSharedPtr<class FOnlineAchievementsNull, ESPMode::ThreadSafe> FOnlineAc
 typedef TSharedPtr<class FOnlineStoreV2Null, ESPMode::ThreadSafe> FOnlineStoreV2NullPtr;
 typedef TSharedPtr<class FOnlinePurchaseNull, ESPMode::ThreadSafe> FOnlinePurchaseNullPtr;
 typedef TSharedPtr<class FMessageSanitizerNull, ESPMode::ThreadSafe> FMessageSanitizerNullPtr;
+typedef TSharedPtr<class FOnlineTitleFileNull, ESPMode::ThreadSafe> FOnlineTitleFileNullPtr;
 #if WITH_ENGINE
 typedef TSharedPtr<class FOnlineVoiceImpl, ESPMode::ThreadSafe> FOnlineVoiceImplPtr;
 #endif //WITH_ENGINE
@@ -97,6 +98,7 @@ PACKAGE_SCOPE:
 		AchievementsInterface(nullptr),
 		StoreV2Interface(nullptr),
 		MessageSanitizerInterface(nullptr),
+		TitleFileInterface(nullptr),
 		OnlineAsyncTaskThreadRunnable(nullptr),
 		OnlineAsyncTaskThread(nullptr)
 	{}
@@ -162,6 +164,9 @@ private:
 	/** Interface for message sanitizing */
 	FMessageSanitizerNullPtr MessageSanitizerInterface;
 
+	/** Interface for title file storage */
+	FOnlineTitleFileNullPtr TitleFileInterface;
+
 	/** Online async task runnable */
 	class FOnlineAsyncTaskManagerNull* OnlineAsyncTaskThreadRunnable;
 
