[TitleFile][Hotfix] Adding support for TitleFilePath=p4://<path> (windows only) to fetch the files from a perforce depot path.

#Synced-CL 50283614

diff --git a/Engine/Plugins/Online/OnlineSubsystemNull/Source/Private/OnlineTitleFileNull.cpp b/Engine/Plugins/Online/OnlineSubsystemNull/Source/Private/OnlineTitleFileNull.cpp
--- a/Engine/Plugins/Online/OnlineSubsystemNull/Source/Private/OnlineTitleFileNull.cpp
+++ b/Engine/Plugins/Online/OnlineSubsystemNull/Source/Private/OnlineTitleFileNull.cpp
@@ -3,6 +3,7 @@
 #include "OnlineTitleFileNull.h"
 #include "Misc/CommandLine.h"
 #include "Misc/ConfigCacheIni.h"
+#include "Misc/MonitoredProcess.h"
 #include "HAL/FileManager.h"
 #include "OnlineSubsystemNull.h"
 
@@ -47,13 +48,81 @@ void FTitleFileNullAsyncLoadLocalDevFile::DoWork()
 	}
 }
 
+void FOnlineTitleFileNull::InitializeLocalFilesFromPerforceIfNeeded()
+{
+	constexpr TCHAR PerforceProtocol[] = TEXT("p4://");
+	constexpr int32 ProtocolLength = UE_ARRAY_COUNT(PerforceProtocol) - 1;
+	constexpr TCHAR HotfixDirBaseName[] = TEXT("Hotfix ");
 
+	if (LocalPath.ToLower().StartsWith(PerforceProtocol))
+	{
+#if PLATFORM_WINDOWS && !UE_BUILD_SHIPPING
+		TArray<FString> OldHotfixDirs;
+		FString FileMatch = FString::Printf(TEXT("%s/%s*"), *FPaths::EngineIntermediateDir(), HotfixDirBaseName);
+		IFileManager::Get().FindFiles(OldHotfixDirs, *FileMatch, /*Files=*/false, /*Dirs=*/true);
+		for (const FString& DirectoryName : OldHotfixDirs)
+		{
+			FString PidString = DirectoryName;
+			PidString.RemoveFromStart(HotfixDirBaseName);
+			int Pid = 0;
+			LexFromString(Pid, *PidString);
+			if (Pid != 0 && !FPlatformProcess::IsApplicationRunning(Pid))
+			{
+				FString FullPath = FPaths::EngineIntermediateDir() / DirectoryName;
+				IPlatformFile::GetPlatformPhysical().DeleteDirectoryRecursively(*FullPath);
+			}
+		}
 
+		// If the path is a perforce path, attempt to use perforce to print the files to a local temp directory
+		FString TempDirectory = FString::Printf(TEXT("%s/%s%d"), *FPaths::EngineIntermediateDir(), HotfixDirBaseName, FPlatformProcess::GetCurrentProcessId());
+		FPaths::NormalizeDirectoryName(TempDirectory);
+		FPaths::MakePlatformFilename(TempDirectory);
+		if (IPlatformFile::GetPlatformPhysical().CreateDirectory(*TempDirectory))
+		{
+			QUICK_SCOPE_CYCLE_COUNTER(STAT_P4FetchHotfixFiles);
+			bShouldDeleteLocalPath = true;
+			UE::HAL::FProcessStartInfo PerforceStartInfo;
+			PerforceStartInfo.bHidden = true;
+			const TCHAR* PerforcePath = (*LocalPath) + ProtocolLength;
+			FString ArgumentsStr = FString::Printf(TEXT("/c \"p4 print -o ^\"%s\\*.ini^\" //%s/*.ini\""), *TempDirectory, PerforcePath);
+			TArray<FString> Results;
+			bool bSuccess = FMonitoredProcess::RunProcessSynchronous(TEXT("cmd"), ArgumentsStr, Results, /*Timeout=*/30.f);
+			FString OutputString = FString::Join(Results, TEXT("\n"));
+			if (bSuccess && !OutputString.Contains(TEXT("no such file(s)")))
+			{
+				LocalPath = TempDirectory;
+			}
+			else
+			{
+				bSuccess = false;
+				IPlatformFile::GetPlatformPhysical().DeleteDirectoryRecursively(*TempDirectory);
+				bShouldDeleteLocalPath = false;
+				ensureMsgf(false, TEXT("Failed to retrieve hotfix files from %s. Perforce result: %s"), *LocalPath, *OutputString);
+				LocalPath.Reset();
+			}
+		}
+#else
+		checkf(PLATFORM_WINDOW, TEXT("Unexpected p4:// TitleFilePath on unsupported platform. Implement FOnlineTitleFileNull::InitializeLocalFilesFromPerforceIfNeeded() if required."));
+#endif //PLATFORM_WINDOWS
+	}
+}
 
 FOnlineTitleFileNull::FOnlineTitleFileNull()
 {
 	// check for the option that is needed for this to work (if it's not, EnumerateFiles will do nothing
 	FParse::Value(FCommandLine::Get(), TEXT("TitleFileLocalPath="), LocalPath);
+	bShouldDeleteLocalPath = false;
+
+	InitializeLocalFilesFromPerforceIfNeeded();
+}
+
+FOnlineTitleFileNull::~FOnlineTitleFileNull()
+{
+	if (bShouldDeleteLocalPath)
+	{
+		IPlatformFile::GetPlatformPhysical().DeleteDirectoryRecursively(*LocalPath);
+		LocalPath.Reset();
+	}
 }
 
 bool FOnlineTitleFileNull::GetFileContents(const FString& FileName, TArray<uint8>& FileContents)
@@ -124,6 +193,7 @@ bool FOnlineTitleFileNull::EnumerateFiles(const FPagedQuery& Page)
 
 	if (!IFileManager::Get().DirectoryExists(*LocalPath))
 	{
+		ensureMsgf(false, TEXT("TitleFilePath does not exist. Path was: %s"), *LocalPath);
 		TriggerOnEnumerateFilesCompleteDelegates(false, TEXT("-TitleFileLocalPath= was specified, but does not exist."));
 		return false;
 	}
@@ -284,4 +354,3 @@ void FOnlineTitleFileNull::FinishReadFileLocalDevFile(FAsyncTask<FTitleFileNullA
 		UE_LOG_ONLINE(Warning, TEXT("failed to read '%s' from given local Dev path!"), *TaskInfo.LocalDevFileName);
 	}
 }
-	
\ No newline at end of file
diff --git a/Engine/Plugins/Online/OnlineSubsystemNull/Source/Private/OnlineTitleFileNull.h b/Engine/Plugins/Online/OnlineSubsystemNull/Source/Private/OnlineTitleFileNull.h
--- a/Engine/Plugins/Online/OnlineSubsystemNull/Source/Private/OnlineTitleFileNull.h
+++ b/Engine/Plugins/Online/OnlineSubsystemNull/Source/Private/OnlineTitleFileNull.h
@@ -43,7 +43,7 @@ class FOnlineTitleFileNull : public IOnlineTitleFile, public TSharedFromThis<FOn
 public:
 
 	FOnlineTitleFileNull();
-	virtual ~FOnlineTitleFileNull() = default;
+	virtual ~FOnlineTitleFileNull();
 
 
 	virtual bool GetFileContents(const FString& FileName, TArray<uint8>& FileContents) override;
@@ -60,6 +60,7 @@ protected:
 	FString LocalPath;
 	TArray<FCloudFileHeader> FileHeaders;
 	TArray<FCloudFile> Files;
+	bool bShouldDeleteLocalPath;
 
 	/** Holds the outstanding task for hitch free loading of local Dev files. */
 	TArray<FAsyncTask<FTitleFileNullAsyncLoadLocalDevFile>* > AsyncLocalDevReads;
@@ -70,6 +71,9 @@ protected:
 	FCloudFile* GetCloudFile(const FString& FileName, bool bCreateIfMissing);
 	/** Completes the async operation for the local dev file read. Similar to above but without any cloud connection. */
 	void FinishReadFileLocalDevFile(FAsyncTask<FTitleFileNullAsyncLoadLocalDevFile>* Task);
+
+	// Special handling for "p4://" title file paths
+	void InitializeLocalFilesFromPerforceIfNeeded();
 };
 
 using FOnlineTitleFileNullPtr = TSharedPtr<FOnlineTitleFileNull, ESPMode::ThreadSafe>;
