[Hotfix] Fixing a bug in the perforce hotfix implementation that would cause us to place the downloaded hotfixes in the wrong location.

#rnx

#Synced-CL 50626103

diff --git a/Engine/Plugins/Online/OnlineSubsystemNull/Source/Private/OnlineTitleFileNull.cpp b/Engine/Plugins/Online/OnlineSubsystemNull/Source/Private/OnlineTitleFileNull.cpp
--- a/Engine/Plugins/Online/OnlineSubsystemNull/Source/Private/OnlineTitleFileNull.cpp
+++ b/Engine/Plugins/Online/OnlineSubsystemNull/Source/Private/OnlineTitleFileNull.cpp
@@ -58,7 +58,7 @@ void FOnlineTitleFileNull::InitializeLocalFilesFromPerforceIfNeeded()
 	{
 #if PLATFORM_WINDOWS && !UE_BUILD_SHIPPING
 		TArray<FString> OldHotfixDirs;
-		FString FileMatch = FString::Printf(TEXT("%s/%s*"), *FPaths::EngineIntermediateDir(), HotfixDirBaseName);
+		FString FileMatch = FPaths::RemoveDuplicateSlashes(FString::Printf(TEXT("%s/%s*"), *FPaths::EngineIntermediateDir(), HotfixDirBaseName));
 		IFileManager::Get().FindFiles(OldHotfixDirs, *FileMatch, /*Files=*/false, /*Dirs=*/true);
 		for (const FString& DirectoryName : OldHotfixDirs)
 		{
@@ -74,7 +74,7 @@ void FOnlineTitleFileNull::InitializeLocalFilesFromPerforceIfNeeded()
 		}
 
 		// If the path is a perforce path, attempt to use perforce to print the files to a local temp directory
-		FString TempDirectory = FString::Printf(TEXT("%s/%s%d"), *FPaths::EngineIntermediateDir(), HotfixDirBaseName, FPlatformProcess::GetCurrentProcessId());
+		FString TempDirectory = FPaths::ConvertRelativePathToFull(FPaths::RemoveDuplicateSlashes(FString::Printf(TEXT("%s/%s%d"), *FPaths::EngineIntermediateDir(), HotfixDirBaseName, FPlatformProcess::GetCurrentProcessId())));
 		FPaths::NormalizeDirectoryName(TempDirectory);
 		FPaths::MakePlatformFilename(TempDirectory);
 		if (IPlatformFile::GetPlatformPhysical().CreateDirectory(*TempDirectory))
