UE: Fix OSS Null not finding multiple LAN sessions from a single host.

FLANSession::Tick cached the LanBeaconGuid from each host's response packets in CachedResponseGuids, so only the first response packet per host was ever processed.
When a single host advertised multiple sessions (e.g.: a game session and a party session), each response packet carried the same LanBeaconGuid, causing every packet after the first to be silently dropped.

Fixes:
- Removed the CachedResponseGuids TSet and its per-host GUID filtering (dead code after the dedup responsibility moved to the session interface layer).
- Added SessionId-based deduplication inFOnlineSessionNull::OnValidResponsePacketReceived so that retry queries still do not produce duplicate search results, while different sessions from the same host are correctly discovered.
- Added an automation test that hosts two sessions on one Null OSS instance and verifies a second instance discovers both via LAN.

#Synced-CL 52126719

diff --git a/Engine/Plugins/Online/OnlineBase/Source/Private/Online/LANBeacon.cpp b/Engine/Plugins/Online/OnlineBase/Source/Private/Online/LANBeacon.cpp
--- a/Engine/Plugins/Online/OnlineBase/Source/Private/Online/LANBeacon.cpp
+++ b/Engine/Plugins/Online/OnlineBase/Source/Private/Online/LANBeacon.cpp
@@ -163,8 +163,6 @@ bool FLANSession::Search(FNboSerializeToBuffer& Packet, FOnValidResponsePacketDe
 	
 	LanBeaconGuid = FGuid::NewGuid();
 
-	CachedResponseGuids.Empty();
-
 	// Bind a socket for LAN beacon activity
 	LanBeacon = new FLanBeacon();
 	if (LanBeacon->Init(LanAnnouncePort) == false)
@@ -222,8 +220,6 @@ void FLANSession::StopLANSession()
 	OnSearchTimeoutDelegates.Clear();
 
 	RetryData.Empty();
-
-	CachedResponseGuids.Empty();
 }
 
 void FLANSession::Tick(float DeltaTime)
@@ -259,14 +255,8 @@ void FLANSession::Tick(float DeltaTime)
 				// We can only accept Server Response packets
 				if (IsValidLanResponsePacket(PacketData, NumRead, ResponseGuid))
 				{
-					bool bAlreadySet = false;
-					CachedResponseGuids.Add(ResponseGuid, &bAlreadySet);
-
-					if (!bAlreadySet)
-					{
-						// Strip off the header
-						OnValidResponsePacketDelegates.Broadcast(&PacketData[LAN_BEACON_PACKET_HEADER_SIZE], NumRead - LAN_BEACON_PACKET_HEADER_SIZE);
-					}
+					// Strip off the header
+					OnValidResponsePacketDelegates.Broadcast(&PacketData[LAN_BEACON_PACKET_HEADER_SIZE], NumRead - LAN_BEACON_PACKET_HEADER_SIZE);
 				}
 			}
 		}
diff --git a/Engine/Plugins/Online/OnlineBase/Source/Public/Online/LANBeacon.h b/Engine/Plugins/Online/OnlineBase/Source/Public/Online/LANBeacon.h
--- a/Engine/Plugins/Online/OnlineBase/Source/Public/Online/LANBeacon.h
+++ b/Engine/Plugins/Online/OnlineBase/Source/Public/Online/LANBeacon.h
@@ -236,9 +236,6 @@ public:
 	/** Unique identifier for this hosting/searching LAN session */
 	FGuid LanBeaconGuid;
 
-	/** Cached guids already received during current search */
-	TSet<FGuid> CachedResponseGuids;
-
 	FLANSession() :
 		LanAnnouncePort(LAN_ANNOUNCE_PORT),
 		LanGameUniqueId(LAN_UNIQUE_ID),
diff --git a/Engine/Plugins/Online/OnlineSubsystemNull/Source/Private/OnlineSessionInterfaceNull.cpp b/Engine/Plugins/Online/OnlineSubsystemNull/Source/Private/OnlineSessionInterfaceNull.cpp
--- a/Engine/Plugins/Online/OnlineSubsystemNull/Source/Private/OnlineSessionInterfaceNull.cpp
+++ b/Engine/Plugins/Online/OnlineSubsystemNull/Source/Private/OnlineSessionInterfaceNull.cpp
@@ -1205,8 +1205,23 @@ void FOnlineSessionNull::OnValidResponsePacketReceived(uint8* PacketData, int32
 
 		if (SearchResult.Session.SessionSettings.BuildUniqueId == BuildUniqueId)
 		{
-			// Add the found session to the search results array
-			CurrentSessionSearch->SearchResults.Emplace(MoveTemp(SearchResult));
+			// Deduplicate by SessionId so that retried queries don't produce duplicate search results, while still allowing 
+			// multiple different sessions from the same host to appear 
+			const FString& NewSessionId = SearchResult.Session.GetSessionIdStr();
+			bool bAlreadyInResults = false;
+			for (const FOnlineSessionSearchResult& Existing : CurrentSessionSearch->SearchResults)
+			{
+				if (Existing.Session.GetSessionIdStr() == NewSessionId)
+				{
+					bAlreadyInResults = true;
+					break;
+				}
+			}
+
+			if (!bAlreadyInResults)
+			{
+				CurrentSessionSearch->SearchResults.Emplace(MoveTemp(SearchResult));
+			}
 		}
 		else
 		{
diff --git a/Engine/Plugins/Online/OnlineSubsystemNull/Source/Private/Tests/OnlineSessionNullLANTest.cpp b/Engine/Plugins/Online/OnlineSubsystemNull/Source/Private/Tests/OnlineSessionNullLANTest.cpp
new file mode 100644
--- /dev/null
+++ b/Engine/Plugins/Online/OnlineSubsystemNull/Source/Private/Tests/OnlineSessionNullLANTest.cpp
@@ -0,0 +1,172 @@
+// Copyright Epic Games, Inc. All Rights Reserved.
+
+#include "CoreMinimal.h"
+#include "Misc/AutomationTest.h"
+#include "OnlineSubsystemNull.h"
+#include "OnlineSessionInterfaceNull.h"
+#include "OnlineSessionSettings.h"
+#include "Online/OnlineSessionNames.h"
+
+#if WITH_DEV_AUTOMATION_TESTS
+
+/**
+ * Verifies that the Null OSS LAN session search can discover multiple sessions
+ * hosted by a single instance. This is a regression test for UE-178220 where
+ * FLANSession::Tick's per-host GUID caching prevented all but the first session
+ * from being found.
+ *
+ * The test creates two FOnlineSubsystemNull instances (Host and Searcher),
+ * registers two sessions on the Host, performs a LAN search from the Searcher,
+ * and asserts that both sessions appear in the results with distinct SessionIds.
+ */
+IMPLEMENT_SIMPLE_AUTOMATION_TEST(
+	FOnlineSessionNullMultiSessionLANDiscovery,
+	"System.Engine.Online.OnlineSubsystemNull.Session.MultiSessionLANDiscovery",
+	EAutomationTestFlags::EditorContext | EAutomationTestFlags::EngineFilter)
+
+bool FOnlineSessionNullMultiSessionLANDiscovery::RunTest(const FString& Parameters)
+{
+	// Create two subsystem instances with distinct names so they get separate LAN beacons.
+	TSharedRef<FOnlineSubsystemNull> HostSubsystem = MakeShared<FOnlineSubsystemNull>(FName(TEXT("TestHost")));
+	TSharedRef<FOnlineSubsystemNull> SearchSubsystem = MakeShared<FOnlineSubsystemNull>(FName(TEXT("TestSearcher")));
+
+	if (!HostSubsystem->Init() || !SearchSubsystem->Init())
+	{
+		AddError(TEXT("Failed to initialize one or both Null subsystem instances."));
+		HostSubsystem->Shutdown();
+		SearchSubsystem->Shutdown();
+		return false;
+	}
+
+	IOnlineSessionPtr HostSessionInt = HostSubsystem->GetSessionInterface();
+	IOnlineSessionPtr SearchSessionInt = SearchSubsystem->GetSessionInterface();
+
+	if (!HostSessionInt.IsValid() || !SearchSessionInt.IsValid())
+	{
+		AddError(TEXT("Failed to get session interface from one or both subsystems."));
+		HostSessionInt.Reset();
+		SearchSessionInt.Reset();
+		HostSubsystem->Shutdown();
+		SearchSubsystem->Shutdown();
+		return false;
+	}
+
+	// --- Host: create two LAN sessions ---
+
+	FOnlineSessionSettings SessionSettingsA;
+	SessionSettingsA.NumPublicConnections = 4;
+	SessionSettingsA.bIsLANMatch = true;
+	SessionSettingsA.bShouldAdvertise = true;
+	SessionSettingsA.bAllowJoinInProgress = true;
+	SessionSettingsA.bAllowJoinViaPresence = true;
+
+	FOnlineSessionSettings SessionSettingsB;
+	SessionSettingsB.NumPublicConnections = 4;
+	SessionSettingsB.bIsLANMatch = true;
+	SessionSettingsB.bShouldAdvertise = true;
+	SessionSettingsB.bAllowJoinInProgress = true;
+	SessionSettingsB.bAllowJoinViaPresence = true;
+
+	const FName SessionNameA = NAME_GameSession;
+	const FName SessionNameB = FName(TEXT("PartySession"));
+
+	bool bCreatedA = HostSessionInt->CreateSession(0, SessionNameA, SessionSettingsA);
+	bool bCreatedB = HostSessionInt->CreateSession(0, SessionNameB, SessionSettingsB);
+
+	if (!bCreatedA || !bCreatedB)
+	{
+		AddError(FString::Printf(TEXT("Failed to create host sessions (A=%d, B=%d)."), bCreatedA, bCreatedB));
+		HostSessionInt.Reset();
+		SearchSessionInt.Reset();
+		HostSubsystem->Shutdown();
+		SearchSubsystem->Shutdown();
+		return false;
+	}
+
+	// --- Searcher: start a LAN search ---
+
+	TSharedRef<FOnlineSessionSearch> SearchSettings = MakeShared<FOnlineSessionSearch>();
+	SearchSettings->bIsLanQuery = true;
+	SearchSettings->MaxSearchResults = 10;
+
+	bool bSearchStarted = SearchSessionInt->FindSessions(0, SearchSettings);
+	if (!bSearchStarted)
+	{
+		AddError(TEXT("Failed to start LAN session search."));
+		HostSessionInt->DestroySession(SessionNameA);
+		HostSessionInt->DestroySession(SessionNameB);
+		HostSessionInt.Reset();
+		SearchSessionInt.Reset();
+		HostSubsystem->Shutdown();
+		SearchSubsystem->Shutdown();
+		return false;
+	}
+
+	// --- Tick both subsystems until the search completes or we time out ---
+
+	const double TimeoutSeconds = 8.0;
+	const double StartTime = FPlatformTime::Seconds();
+	const float TickDelta = 0.016f; // ~60fps
+
+	while (SearchSettings->SearchState == EOnlineAsyncTaskState::InProgress)
+	{
+		HostSubsystem->Tick(TickDelta);
+		SearchSubsystem->Tick(TickDelta);
+
+		FPlatformProcess::Sleep(TickDelta);
+
+		if (FPlatformTime::Seconds() - StartTime > TimeoutSeconds)
+		{
+			AddError(TEXT("LAN search timed out before completing."));
+			break;
+		}
+	}
+
+	// --- Validate results ---
+
+	const int32 NumResults = SearchSettings->SearchResults.Num();
+	UE_LOG(LogTemp, Log, TEXT("[UE-178220 Test] Search returned %d result(s)."), NumResults);
+
+	if (NumResults < 2)
+	{
+		AddError(FString::Printf(
+			TEXT("Expected at least 2 search results but got %d. "
+			     "This indicates the multi-session LAN discovery fix may not be working."),
+			NumResults));
+	}
+	else
+	{
+		// Verify distinct SessionIds
+		TSet<FString> UniqueSessionIds;
+		for (const FOnlineSessionSearchResult& Result : SearchSettings->SearchResults)
+		{
+			FString SessionId = Result.Session.GetSessionIdStr();
+			UE_LOG(LogTemp, Log, TEXT("[UE-178220 Test] Found session: %s"), *SessionId);
+			UniqueSessionIds.Add(SessionId);
+		}
+
+		if (UniqueSessionIds.Num() < 2)
+		{
+			AddError(TEXT("Search results contain duplicate SessionIds - expected distinct sessions."));
+		}
+		else
+		{
+			UE_LOG(LogTemp, Log, TEXT("[UE-178220 Test] PASSED: Found %d unique sessions."), UniqueSessionIds.Num());
+		}
+	}
+
+	// --- Cleanup ---
+	// Release interface refs before Shutdown(), which asserts unique ownership.
+
+	HostSessionInt->DestroySession(SessionNameA);
+	HostSessionInt->DestroySession(SessionNameB);
+	HostSessionInt.Reset();
+	SearchSessionInt.Reset();
+
+	HostSubsystem->Shutdown();
+	SearchSubsystem->Shutdown();
+
+	return true;
+}
+
+#endif // WITH_DEV_AUTOMATION_TESTS
